Added API exposure to some Blutility classes

Requested from AnswerHub: https://answers.unrealengine.com/questions/33831/how-do-i-extend-uglobaleditorutilitybase-blutility.html

Checked with Michael Noland and have added the BLUTILITY_API macros to UGlobalEditorUtilityBase and APlacedEditorUtilityBase

diff --git a/Engine/Source/Editor/Blutility/Classes/GlobalEditorUtilityBase.h b/Engine/Source/Editor/Blutility/Classes/GlobalEditorUtilityBase.h
--- a/Engine/Source/Editor/Blutility/Classes/GlobalEditorUtilityBase.h
+++ b/Engine/Source/Editor/Blutility/Classes/GlobalEditorUtilityBase.h
@@ -14,7 +14,7 @@ DECLARE_DYNAMIC_MULTICAST_DELEGATE_TwoParams( FForEachAssetIteratorSignature, cl
 
 
 UCLASS(Abstract, hideCategories=(Object), Blueprintable)
-class UGlobalEditorUtilityBase : public UObject
+class BLUTILITY_API UGlobalEditorUtilityBase : public UObject
 {
 	GENERATED_UCLASS_BODY()
 
diff --git a/Engine/Source/Editor/Blutility/Classes/PlacedEditorUtilityBase.h b/Engine/Source/Editor/Blutility/Classes/PlacedEditorUtilityBase.h
--- a/Engine/Source/Editor/Blutility/Classes/PlacedEditorUtilityBase.h
+++ b/Engine/Source/Editor/Blutility/Classes/PlacedEditorUtilityBase.h
@@ -9,7 +9,7 @@
 #include "PlacedEditorUtilityBase.generated.h"
 
 UCLASS(Abstract, hideCategories=(Object, Actor)/*, Blueprintable*/)
-class APlacedEditorUtilityBase : public AActor
+class BLUTILITY_API APlacedEditorUtilityBase : public AActor
 {
 	GENERATED_UCLASS_BODY()
 
