#ue4 Adding support for grabbing the selection set and the current perspective viewport transform in PlacedEditorUtility classes

diff --git a/Engine/Source/Editor/Blutility/Classes/PlacedEditorUtilityBase.h b/Engine/Source/Editor/Blutility/Classes/PlacedEditorUtilityBase.h
--- a/Engine/Source/Editor/Blutility/Classes/PlacedEditorUtilityBase.h
+++ b/Engine/Source/Editor/Blutility/Classes/PlacedEditorUtilityBase.h
@@ -19,4 +19,20 @@ class BLUTILITY_API APlacedEditorUtilityBase : public AActor
 	// AActor interface
 	virtual void TickActor(float DeltaSeconds, ELevelTick TickType, FActorTickFunction& ThisTickFunction) override;
 	// End of AActor interface
+
+	/**
+	 * Returns the current selection set in the editor.  Note that for non-editor builds, this will return an empty array
+	 */
+	UFUNCTION(BlueprintPure, Category = "Development|Editor")
+	TArray<AActor*> GetSelectionSet();
+
+	/**
+	 * Gets information about the camera position for the primary level editor viewport.  In non-editor builds, these will be zeroed
+	 *
+	 * @param	CameraLocation	(out) Current location of the level editing viewport camera, or zero if none found
+	 * @param	CameraRotation	(out) Current rotation of the level editing viewport camera, or zero if none found
+	 * @return	Whether or not we were able to get a camera for a level editing viewport
+	 */
+	UFUNCTION(BlueprintPure, Category = "Development|Editor")
+	bool GetLevelViewportCameraInfo(FVector& CameraLocation, FRotator& CameraRotation);
 };
diff --git a/Engine/Source/Editor/Blutility/Private/PlacedEditorUtilityBase.cpp b/Engine/Source/Editor/Blutility/Private/PlacedEditorUtilityBase.cpp
--- a/Engine/Source/Editor/Blutility/Private/PlacedEditorUtilityBase.cpp
+++ b/Engine/Source/Editor/Blutility/Private/PlacedEditorUtilityBase.cpp
@@ -25,6 +25,46 @@ void APlacedEditorUtilityBase::TickActor(float DeltaSeconds, ELevelTick TickType
 	Super::TickActor(DeltaSeconds, TickType, ThisTickFunction);
 }
 
+TArray<AActor*> APlacedEditorUtilityBase::GetSelectionSet()
+{
+	TArray<AActor*> Result;
+
+#if WITH_EDITOR
+	for (FSelectionIterator It(GEditor->GetSelectedActorIterator()); It; ++It)
+	{
+		if (AActor* Actor = Cast<AActor>(*It))
+		{
+			Result.Add(Actor);
+		}
+	}
+#endif //WITH_EDITOR
+
+	return Result;
+}
+
+bool APlacedEditorUtilityBase::GetLevelViewportCameraInfo(FVector& CameraLocation, FRotator& CameraRotation)
+{
+	bool RetVal = false;
+	CameraLocation = FVector::ZeroVector;
+	CameraRotation = FRotator::ZeroRotator;
+
+#if WITH_EDITOR
+	for (auto LevelVC : GEditor->LevelViewportClients)
+	{
+		if (LevelVC && LevelVC->IsPerspective())
+		{
+			CameraLocation = LevelVC->GetViewLocation();
+			CameraRotation = LevelVC->GetViewRotation();
+			RetVal = true;
+
+			break;
+		}
+	}
+#endif //WITH_EDITOR
+
+	return RetVal;
+}
+
 /*
 
 
