#ue4 Adding soft reference support to blutilities

diff --git a/Engine/Source/Editor/Blutility/Classes/GlobalEditorUtilityBase.h b/Engine/Source/Editor/Blutility/Classes/GlobalEditorUtilityBase.h
--- a/Engine/Source/Editor/Blutility/Classes/GlobalEditorUtilityBase.h
+++ b/Engine/Source/Editor/Blutility/Classes/GlobalEditorUtilityBase.h
@@ -27,6 +27,14 @@ class BLUTILITY_API UGlobalEditorUtilityBase : public UObject
 	UFUNCTION(BlueprintCallable, Category="Development|Editor")
 	TArray<AActor*> GetSelectionSet();
 
+	/**
+	 * Attempts to find the actor specified by PathToActor in the current editor world
+	 * @param	PathToActor	The path to the actor (e.g. PersistentLevel.PlayerStart)
+	 * @return	A reference to the actor, or none if it wasn't found
+	 */
+	UFUNCTION(BlueprintPure, Category = "Development|Editor")
+	AActor* GetActorReference(FString PathToActor);
+
 	////////////////////////////
 
 	// Should this blueprint automatically run OnDefaultActionClicked, or should it open up a details panel to edit properties and/or offer multiple buttons
diff --git a/Engine/Source/Editor/Blutility/Private/GlobalEditorUtilityBase.cpp b/Engine/Source/Editor/Blutility/Private/GlobalEditorUtilityBase.cpp
--- a/Engine/Source/Editor/Blutility/Private/GlobalEditorUtilityBase.cpp
+++ b/Engine/Source/Editor/Blutility/Private/GlobalEditorUtilityBase.cpp
@@ -144,4 +144,13 @@ void UGlobalEditorUtilityBase::RenameAsset(UObject* Asset, const FString& NewNam
 	new (AssetsAndNames) FAssetRenameData(Asset, PackagePath, NewName);
 
 	AssetToolsModule.Get().RenameAssets(AssetsAndNames);
+}
+
+AActor* UGlobalEditorUtilityBase::GetActorReference(FString PathToActor)
+{
+#if WITH_EDITOR
+	return Cast<AActor>(StaticFindObject(AActor::StaticClass(), GEditor->GetEditorWorldContext().World(), *PathToActor, false));
+#else
+	return nullptr;
+#endif //WITH_EDITOR
 }
\ No newline at end of file
