#UE4 DOCS: copied over more selection functionality from the GlobalEditorUtility to the PlacedEditorUtility. Placed Utilities can now select actors and clear the selection.

diff --git a/Engine/Source/Editor/Blutility/Classes/PlacedEditorUtilityBase.h b/Engine/Source/Editor/Blutility/Classes/PlacedEditorUtilityBase.h
--- a/Engine/Source/Editor/Blutility/Classes/PlacedEditorUtilityBase.h
+++ b/Engine/Source/Editor/Blutility/Classes/PlacedEditorUtilityBase.h
@@ -35,4 +35,12 @@ class BLUTILITY_API APlacedEditorUtilityBase : public AActor
 	 */
 	UFUNCTION(BlueprintPure, Category = "Development|Editor")
 	bool GetLevelViewportCameraInfo(FVector& CameraLocation, FRotator& CameraRotation);
+
+	// Remove all actors from the selection set
+	UFUNCTION(BlueprintCallable, Category = "Development|Editor")
+		void ClearActorSelectionSet();
+
+	// Set the selection state for the selected actor
+	UFUNCTION(BlueprintCallable, Category = "Development|Editor")
+		void SetActorSelectionState(AActor* Actor, bool bShouldBeSelected);
 };
diff --git a/Engine/Source/Editor/Blutility/Private/GlobalEditorUtilityBase.cpp b/Engine/Source/Editor/Blutility/Private/GlobalEditorUtilityBase.cpp
--- a/Engine/Source/Editor/Blutility/Private/GlobalEditorUtilityBase.cpp
+++ b/Engine/Source/Editor/Blutility/Private/GlobalEditorUtilityBase.cpp
@@ -102,7 +102,7 @@ UEditorUserSettings* UGlobalEditorUtilityBase::GetEditorUserSettings()
 
 void UGlobalEditorUtilityBase::ClearActorSelectionSet()
 {
-	GEditor->GetSelectedActors()->DeselectAll();
+	GEditor->SelectNone(true, true, false);
 	bDirtiedSelectionSet = true;
 }
 
diff --git a/Engine/Source/Editor/Blutility/Private/PlacedEditorUtilityBase.cpp b/Engine/Source/Editor/Blutility/Private/PlacedEditorUtilityBase.cpp
--- a/Engine/Source/Editor/Blutility/Private/PlacedEditorUtilityBase.cpp
+++ b/Engine/Source/Editor/Blutility/Private/PlacedEditorUtilityBase.cpp
@@ -65,6 +65,18 @@ bool APlacedEditorUtilityBase::GetLevelViewportCameraInfo(FVector& CameraLocatio
 	return RetVal;
 }
 
+void APlacedEditorUtilityBase::ClearActorSelectionSet()
+{
+	//GEditor->GetSelectedActors()->DeselectAll();
+		
+	GEditor->SelectNone(true, true, false);
+}
+
+void APlacedEditorUtilityBase::SetActorSelectionState(AActor* Actor, bool bShouldBeSelected)
+{
+	GEditor->SelectActor(Actor, bShouldBeSelected, /*bNotify=*/ false);
+}
+
 /*
 
 
