#UE4 DOCS: Added a Blueprint function to the PlacedEditorUtility to set the editor's 3D viewport camera's location and rotation

diff --git a/Engine/Source/Editor/Blutility/Classes/PlacedEditorUtilityBase.h b/Engine/Source/Editor/Blutility/Classes/PlacedEditorUtilityBase.h
--- a/Engine/Source/Editor/Blutility/Classes/PlacedEditorUtilityBase.h
+++ b/Engine/Source/Editor/Blutility/Classes/PlacedEditorUtilityBase.h
@@ -36,11 +36,28 @@ class BLUTILITY_API APlacedEditorUtilityBase : public AActor
 	UFUNCTION(BlueprintPure, Category = "Development|Editor")
 	bool GetLevelViewportCameraInfo(FVector& CameraLocation, FRotator& CameraRotation);
 
+	/**
+	* Sets information about the camera position for the primary level editor viewport.
+	*
+	* @param	CameraLocation	Location the camera will be moved to.
+	* @param	CameraRotation	Rotation the camera will be set to.
+	*/
+	UFUNCTION(BlueprintCallable, Category = "Development|Editor")
+	void SetLevelViewportCameraInfo(FVector CameraLocation, FRotator CameraRotation);
+
 	// Remove all actors from the selection set
 	UFUNCTION(BlueprintCallable, Category = "Development|Editor")
 		void ClearActorSelectionSet();
 
 	// Set the selection state for the selected actor
 	UFUNCTION(BlueprintCallable, Category = "Development|Editor")
-		void SetActorSelectionState(AActor* Actor, bool bShouldBeSelected);
+	void SetActorSelectionState(AActor* Actor, bool bShouldBeSelected);
+
+	/**
+	* Attempts to find the actor specified by PathToActor in the current editor world
+	* @param	PathToActor	The path to the actor (e.g. PersistentLevel.PlayerStart)
+	* @return	A reference to the actor, or none if it wasn't found
+	*/
+	UFUNCTION(BlueprintPure, Category = "Development|Editor")
+	AActor* GetActorReference(FString PathToActor);
 };
diff --git a/Engine/Source/Editor/Blutility/Private/PlacedEditorUtilityBase.cpp b/Engine/Source/Editor/Blutility/Private/PlacedEditorUtilityBase.cpp
--- a/Engine/Source/Editor/Blutility/Private/PlacedEditorUtilityBase.cpp
+++ b/Engine/Source/Editor/Blutility/Private/PlacedEditorUtilityBase.cpp
@@ -65,6 +65,26 @@ bool APlacedEditorUtilityBase::GetLevelViewportCameraInfo(FVector& CameraLocatio
 	return RetVal;
 }
 
+
+void APlacedEditorUtilityBase::SetLevelViewportCameraInfo(FVector CameraLocation, FRotator CameraRotation)
+{
+
+#if WITH_EDITOR
+	for (auto LevelVC : GEditor->LevelViewportClients)
+	{
+		if (LevelVC && LevelVC->IsPerspective())
+		{
+			LevelVC->SetViewLocation(CameraLocation);
+			LevelVC->SetViewRotation(CameraRotation);
+
+			break;
+		}
+	}
+#endif //WITH_EDITOR
+}
+
+
+
 void APlacedEditorUtilityBase::ClearActorSelectionSet()
 {
 	//GEditor->GetSelectedActors()->DeselectAll();
@@ -77,6 +97,16 @@ void APlacedEditorUtilityBase::SetActorSelectionState(AActor* Actor, bool bShoul
 	GEditor->SelectActor(Actor, bShouldBeSelected, /*bNotify=*/ false);
 }
 
+AActor* APlacedEditorUtilityBase::GetActorReference(FString PathToActor)
+{
+#if WITH_EDITOR
+	return Cast<AActor>(StaticFindObject(AActor::StaticClass(), GEditor->GetEditorWorldContext().World(), *PathToActor, false));
+#else
+	return nullptr;
+#endif //WITH_EDITOR
+}
+
+
 /*
 
 
