#UE4 DOCS: rolled back changes to ClearActorSelectionSet(), added new function SelectNothing(); to both Global and Placed Editor Utilities

diff --git a/Engine/Source/Editor/Blutility/Classes/GlobalEditorUtilityBase.h b/Engine/Source/Editor/Blutility/Classes/GlobalEditorUtilityBase.h
--- a/Engine/Source/Editor/Blutility/Classes/GlobalEditorUtilityBase.h
+++ b/Engine/Source/Editor/Blutility/Classes/GlobalEditorUtilityBase.h
@@ -74,6 +74,10 @@ class BLUTILITY_API UGlobalEditorUtilityBase : public UObject
 	UFUNCTION(BlueprintCallable, Category="Development|Editor")
 	void ClearActorSelectionSet();
 
+	// Selects nothing in the editor (another way to clear the selection)
+	UFUNCTION(BlueprintCallable, Category = "Development|Editor")
+	void SelectNothing();
+
 	// Set the selection state for the selected actor
 	UFUNCTION(BlueprintCallable, Category="Development|Editor")
 	void SetActorSelectionState(AActor* Actor, bool bShouldBeSelected);
diff --git a/Engine/Source/Editor/Blutility/Classes/PlacedEditorUtilityBase.h b/Engine/Source/Editor/Blutility/Classes/PlacedEditorUtilityBase.h
--- a/Engine/Source/Editor/Blutility/Classes/PlacedEditorUtilityBase.h
+++ b/Engine/Source/Editor/Blutility/Classes/PlacedEditorUtilityBase.h
@@ -47,7 +47,11 @@ class BLUTILITY_API APlacedEditorUtilityBase : public AActor
 
 	// Remove all actors from the selection set
 	UFUNCTION(BlueprintCallable, Category = "Development|Editor")
-		void ClearActorSelectionSet();
+	void ClearActorSelectionSet();
+
+	// Selects nothing in the editor (another way to clear the selection)
+	UFUNCTION(BlueprintCallable, Category = "Development|Editor")
+	void SelectNothing();
 
 	// Set the selection state for the selected actor
 	UFUNCTION(BlueprintCallable, Category = "Development|Editor")
diff --git a/Engine/Source/Editor/Blutility/Private/GlobalEditorUtilityBase.cpp b/Engine/Source/Editor/Blutility/Private/GlobalEditorUtilityBase.cpp
--- a/Engine/Source/Editor/Blutility/Private/GlobalEditorUtilityBase.cpp
+++ b/Engine/Source/Editor/Blutility/Private/GlobalEditorUtilityBase.cpp
@@ -101,6 +101,12 @@ UEditorUserSettings* UGlobalEditorUtilityBase::GetEditorUserSettings()
 }
 
 void UGlobalEditorUtilityBase::ClearActorSelectionSet()
+{
+	GEditor->GetSelectedActors()->DeselectAll();
+	bDirtiedSelectionSet = true;
+}
+
+void UGlobalEditorUtilityBase::SelectNothing()
 {
 	GEditor->SelectNone(true, true, false);
 	bDirtiedSelectionSet = true;
diff --git a/Engine/Source/Editor/Blutility/Private/PlacedEditorUtilityBase.cpp b/Engine/Source/Editor/Blutility/Private/PlacedEditorUtilityBase.cpp
--- a/Engine/Source/Editor/Blutility/Private/PlacedEditorUtilityBase.cpp
+++ b/Engine/Source/Editor/Blutility/Private/PlacedEditorUtilityBase.cpp
@@ -83,12 +83,14 @@ void APlacedEditorUtilityBase::SetLevelViewportCameraInfo(FVector CameraLocation
 #endif //WITH_EDITOR
 }
 
-
-
 void APlacedEditorUtilityBase::ClearActorSelectionSet()
 {
-	//GEditor->GetSelectedActors()->DeselectAll();
-		
+	GEditor->GetSelectedActors()->DeselectAll();
+	GEditor->NoteSelectionChange();
+}
+
+void APlacedEditorUtilityBase::SelectNothing()
+{
 	GEditor->SelectNone(true, true, false);
 }
 
