Fixed some places that were running asset registry queries inside a loop and causing performance issues

diff --git a/Engine/Plugins/Experimental/ControlRig/Source/ControlRigEditor/Private/ControlRigEditorModule.cpp b/Engine/Plugins/Experimental/ControlRig/Source/ControlRigEditor/Private/ControlRigEditorModule.cpp
--- a/Engine/Plugins/Experimental/ControlRig/Source/ControlRigEditor/Private/ControlRigEditorModule.cpp
+++ b/Engine/Plugins/Experimental/ControlRig/Source/ControlRigEditor/Private/ControlRigEditorModule.cpp
@@ -188,18 +188,34 @@ void FControlRigEditorModule::StartupModule()
 			FAssetRegistryModule& AssetRegistryModule = FModuleManager::LoadModuleChecked<FAssetRegistryModule>("AssetRegistry");
 			
 			bool bCanReimport = false;
-			for(const FAssetData& AssetData : SelectedAssets)
+			if (SelectedAssets.Num() > 0)
 			{
-				TMultiMap<FName, FString> TagsAndValues;
-				TagsAndValues.Add(GET_MEMBER_NAME_CHECKED(UControlRigSequence, LastExportedToAnimationSequence), AssetData.ObjectPath.ToString());
-
+				// It's faster to find all assets with this tag and then query them against the selection then it is to 
+				// query the asset registry each time for a tag with a particular value
+				const FName LastExportedToAnimationSequenceTagName = GET_MEMBER_NAME_CHECKED(UControlRigSequence, LastExportedToAnimationSequence);
 				TArray<FAssetData> FoundAssets;
-				AssetRegistryModule.Get().GetAssetsByTagValues(TagsAndValues, FoundAssets);
+				{
+					TArray<FName> Tags;
+					Tags.Add(LastExportedToAnimationSequenceTagName);
+					AssetRegistryModule.Get().GetAssetsByTags(Tags, FoundAssets);
+				}
 
 				if (FoundAssets.Num() > 0)
 				{
-					bCanReimport = true;
-					break;
+					for(const FAssetData& AssetData : SelectedAssets)
+					{
+						const bool bFoundAsset = FoundAssets.ContainsByPredicate([&AssetData, LastExportedToAnimationSequenceTagName](const FAssetData& FoundAsset)
+						{
+							const FName TagValue = FoundAsset.GetTagValueRef<FName>(LastExportedToAnimationSequenceTagName);
+							return TagValue == AssetData.ObjectPath;
+						});
+
+						if (bFoundAsset)
+						{
+							bCanReimport = true;
+							break;
+						}
+					}
 				}
 			}
 
diff --git a/Engine/Source/Editor/Blutility/Private/BlutilityContentBrowserExtensions.cpp b/Engine/Source/Editor/Blutility/Private/BlutilityContentBrowserExtensions.cpp
--- a/Engine/Source/Editor/Blutility/Private/BlutilityContentBrowserExtensions.cpp
+++ b/Engine/Source/Editor/Blutility/Private/BlutilityContentBrowserExtensions.cpp
@@ -36,26 +36,29 @@ public:
 
 		// Run thru the assets to determine if any meet our criteria
 		TArray<UGlobalEditorUtilityBase*> SupportedUtils;
-		for (auto AssetIt = SelectedAssets.CreateConstIterator(); AssetIt; ++AssetIt)
+		if (SelectedAssets.Num() > 0)
 		{
-			const FAssetData& Asset = *AssetIt;
-			
 			// Check blueprint utils (we need to load them to query their validity against these assets)
 			TArray<FAssetData> UtilAssets;
 			FBlutilityMenuExtensions::GetBlutilityClasses(UtilAssets, UAssetActionUtility::StaticClass()->GetFName());
-
-			for(FAssetData& UtilAsset : UtilAssets)
+			if (UtilAssets.Num() > 0)
 			{
-				if(UEditorUtilityBlueprint* Blueprint = Cast<UEditorUtilityBlueprint>(UtilAsset.GetAsset()))
+				for (const FAssetData& Asset : SelectedAssets)
 				{
-					if(UClass* BPClass = Blueprint->GeneratedClass.Get())
+					for (const FAssetData& UtilAsset : UtilAssets)
 					{
-						if(UAssetActionUtility* DefaultObject = Cast<UAssetActionUtility>(BPClass->GetDefaultObject()))
+						if(UEditorUtilityBlueprint* Blueprint = Cast<UEditorUtilityBlueprint>(UtilAsset.GetAsset()))
 						{
-							UClass* SupportedClass = DefaultObject->GetSupportedClass();
-							if(SupportedClass == nullptr || (SupportedClass && Asset.GetClass()->IsChildOf(SupportedClass)))
+							if(UClass* BPClass = Blueprint->GeneratedClass.Get())
 							{
-								SupportedUtils.AddUnique(DefaultObject);
+								if(UAssetActionUtility* DefaultObject = Cast<UAssetActionUtility>(BPClass->GetDefaultObject()))
+								{
+									UClass* SupportedClass = DefaultObject->GetSupportedClass();
+									if(SupportedClass == nullptr || (SupportedClass && Asset.GetClass()->IsChildOf(SupportedClass)))
+									{
+										SupportedUtils.AddUnique(DefaultObject);
+									}
+								}
 							}
 						}
 					}
diff --git a/Engine/Source/Editor/Blutility/Private/BlutilityLevelEditorExtensions.cpp b/Engine/Source/Editor/Blutility/Private/BlutilityLevelEditorExtensions.cpp
--- a/Engine/Source/Editor/Blutility/Private/BlutilityLevelEditorExtensions.cpp
+++ b/Engine/Source/Editor/Blutility/Private/BlutilityLevelEditorExtensions.cpp
@@ -28,26 +28,31 @@ public:
 
 		// Run thru the assets to determine if any meet our criteria
 		TArray<UGlobalEditorUtilityBase*> SupportedUtils;
-		for (AActor* Actor : SelectedActors)
-		{			
-			if(Actor)
+		if (SelectedActors.Num() > 0)
+		{
+			// Check blueprint utils (we need to load them to query their validity against these actors)
+			TArray<FAssetData> UtilAssets;
+			FBlutilityMenuExtensions::GetBlutilityClasses(UtilAssets, UActorActionUtility::StaticClass()->GetFName());
+			if (UtilAssets.Num() > 0)
 			{
-				// Check blueprint utils (we need to load them to query their validity against these actors)
-				TArray<FAssetData> UtilAssets;
-				FBlutilityMenuExtensions::GetBlutilityClasses(UtilAssets, UActorActionUtility::StaticClass()->GetFName());
-
-				for(FAssetData& UtilAsset : UtilAssets)
-				{
-					if(UEditorUtilityBlueprint* Blueprint = Cast<UEditorUtilityBlueprint>(UtilAsset.GetAsset()))
+				for (AActor* Actor : SelectedActors)
+				{			
+					if(Actor)
 					{
-						if(UClass* BPClass = Blueprint->GeneratedClass.Get())
+						for(FAssetData& UtilAsset : UtilAssets)
 						{
-							if(UActorActionUtility* DefaultObject = Cast<UActorActionUtility>(BPClass->GetDefaultObject()))
+							if(UEditorUtilityBlueprint* Blueprint = Cast<UEditorUtilityBlueprint>(UtilAsset.GetAsset()))
 							{
-								UClass* SupportedClass = DefaultObject->GetSupportedClass();
-								if(SupportedClass == nullptr || (SupportedClass && Actor->GetClass()->IsChildOf(SupportedClass)))
+								if(UClass* BPClass = Blueprint->GeneratedClass.Get())
 								{
-									SupportedUtils.AddUnique(DefaultObject);
+									if(UActorActionUtility* DefaultObject = Cast<UActorActionUtility>(BPClass->GetDefaultObject()))
+									{
+										UClass* SupportedClass = DefaultObject->GetSupportedClass();
+										if(SupportedClass == nullptr || (SupportedClass && Actor->GetClass()->IsChildOf(SupportedClass)))
+										{
+											SupportedUtils.AddUnique(DefaultObject);
+										}
+									}
 								}
 							}
 						}
