Placed Editor Utility Base once again ticks in the editor as intended

diff --git a/Engine/Source/Editor/Blutility/Classes/PlacedEditorUtilityBase.h b/Engine/Source/Editor/Blutility/Classes/PlacedEditorUtilityBase.h
--- a/Engine/Source/Editor/Blutility/Classes/PlacedEditorUtilityBase.h
+++ b/Engine/Source/Editor/Blutility/Classes/PlacedEditorUtilityBase.h
@@ -21,6 +21,7 @@ class BLUTILITY_API APlacedEditorUtilityBase : public AActor
 
 	// AActor interface
 	virtual void TickActor(float DeltaSeconds, ELevelTick TickType, FActorTickFunction& ThisTickFunction) override;
+	virtual bool ShouldTickIfViewportsOnly() const override { return true; }
 	// End of AActor interface
 
 	/**
diff --git a/Engine/Source/Editor/Blutility/Private/PlacedEditorUtilityBase.cpp b/Engine/Source/Editor/Blutility/Private/PlacedEditorUtilityBase.cpp
--- a/Engine/Source/Editor/Blutility/Private/PlacedEditorUtilityBase.cpp
+++ b/Engine/Source/Editor/Blutility/Private/PlacedEditorUtilityBase.cpp
@@ -17,13 +17,7 @@ APlacedEditorUtilityBase::APlacedEditorUtilityBase(const FObjectInitializer& Obj
 
 void APlacedEditorUtilityBase::TickActor(float DeltaSeconds, ELevelTick TickType, FActorTickFunction& ThisTickFunction)
 {
-	// Force us to tick even in the editor viewport
-	if ((TickType == LEVELTICK_ViewportsOnly) && !IsPendingKill())
-	{
-		FEditorScriptExecutionGuard ScriptGuard;
-		ReceiveTick(DeltaSeconds);
-	}
-
+	FEditorScriptExecutionGuard ScriptGuard;
 	Super::TickActor(DeltaSeconds, TickType, ThisTickFunction);
 }
 
@@ -51,7 +45,7 @@ bool APlacedEditorUtilityBase::GetLevelViewportCameraInfo(FVector& CameraLocatio
 	CameraRotation = FRotator::ZeroRotator;
 
 #if WITH_EDITOR
-	for (auto LevelVC : GEditor->LevelViewportClients)
+	for (FLevelEditorViewportClient* LevelVC : GEditor->LevelViewportClients)
 	{
 		if (LevelVC && LevelVC->IsPerspective())
 		{
@@ -72,7 +66,7 @@ void APlacedEditorUtilityBase::SetLevelViewportCameraInfo(FVector CameraLocation
 {
 
 #if WITH_EDITOR
-	for (auto LevelVC : GEditor->LevelViewportClients)
+	for (FLevelEditorViewportClient* LevelVC : GEditor->LevelViewportClients)
 	{
 		if (LevelVC && LevelVC->IsPerspective())
 		{
@@ -109,24 +103,3 @@ AActor* APlacedEditorUtilityBase::GetActorReference(FString PathToActor)
 	return nullptr;
 #endif //WITH_EDITOR
 }
-
-
-/*
-
-
-FNotificationInfo ErrorNotification(TEXT(""));
-ErrorNotification.Image = FEditorStyle::GetBrush(TEXT("MessageLog.Error"));
-ErrorNotification.bFireAndForget = true;
-ErrorNotification.Hyperlink = FSimpleDelegate::CreateRaw(&Local::OpenMessageLog);
-ErrorNotification.HyperlinkText = LOCTEXT("ShowMessageLogHyperlink", "Show Message Log");
-ErrorNotification.ExpireDuration = 3.0f; // Need this message to last a little longer than normal since the user may want to "Show Log"
-ErrorNotification.bUseThrobber = true;
-ErrorNotification.Text = stuff;
-
-void APlacedEditorUtilityBase::ShowNotification()
-// If the actor couldn't be added, display a notification to the user explaining why.
-IMainFrameModule& MainFrameModule = FModuleManager::GetModuleChecked<IMainFrameModule>( TEXT("MainFrame") );
-
-MainFrameModule.AddNotification(ErrorNotification);
-
-*/
