Fix additional params being displayed on blutility asset/actor actions dialog

diff --git a/Engine/Source/Editor/Blutility/Private/BlutilityMenuExtensions.cpp b/Engine/Source/Editor/Blutility/Private/BlutilityMenuExtensions.cpp
--- a/Engine/Source/Editor/Blutility/Private/BlutilityMenuExtensions.cpp
+++ b/Engine/Source/Editor/Blutility/Private/BlutilityMenuExtensions.cpp
@@ -66,6 +66,13 @@ class SFunctionParamDialog : public SCompoundWidget
 		FPropertyEditorModule& PropertyEditorModule = FModuleManager::Get().LoadModuleChecked<FPropertyEditorModule>("PropertyEditor");
 		TSharedRef<IStructureDetailsView> StructureDetailsView = PropertyEditorModule.CreateStructureDetailView(DetailsViewArgs, StructureViewArgs, InStructOnScope);
 
+		StructureDetailsView->GetDetailsView()->SetIsPropertyVisibleDelegate(FIsPropertyVisible::CreateLambda([](const FPropertyAndParent& InPropertyAndParent)
+		{
+			return InPropertyAndParent.Property.HasAnyPropertyFlags(CPF_Parm);
+		}));
+
+		StructureDetailsView->GetDetailsView()->ForceRefresh();
+
 		ChildSlot
 		[
 			SNew(SVerticalBox)
