Making editor utility edit/run behavior more consistent

diff --git a/Engine/Source/Editor/Blutility/Private/AssetTypeActions_EditorUtilityWidgetBlueprint.cpp b/Engine/Source/Editor/Blutility/Private/AssetTypeActions_EditorUtilityWidgetBlueprint.cpp
--- a/Engine/Source/Editor/Blutility/Private/AssetTypeActions_EditorUtilityWidgetBlueprint.cpp
+++ b/Engine/Source/Editor/Blutility/Private/AssetTypeActions_EditorUtilityWidgetBlueprint.cpp
@@ -49,11 +49,11 @@ void FAssetTypeActions_EditorUtilityWidgetBlueprint::GetActions(const TArray<UOb
 	auto Blueprints = GetTypedWeakObjectPtrs<UWidgetBlueprint>(InObjects);
 
 	MenuBuilder.AddMenuEntry(
-		LOCTEXT("EditorUtilityWidget_Edit", "Edit Blueprint"),
-		LOCTEXT("EditorUtilityWidget_EditTooltip", "Opens the selected blueprints in the full blueprint editor."),
+		LOCTEXT("EditorUtilityWidget_Edit", "Run Editor Utility Widget"),
+		LOCTEXT("EditorUtilityWidget_EditTooltip", "Runs the single action or opens the tab built by this Editor Utility Widget Blueprint."),
 		FSlateIcon(),
 		FUIAction(
-			FExecuteAction::CreateSP(this, &FAssetTypeActions_EditorUtilityWidgetBlueprint::ExecuteEdit, Blueprints),
+			FExecuteAction::CreateSP(this, &FAssetTypeActions_EditorUtilityWidgetBlueprint::ExecuteRun, Blueprints),
 			FCanExecuteAction()
 		)
 	);
@@ -64,6 +64,32 @@ void FAssetTypeActions_EditorUtilityWidgetBlueprint::OpenAssetEditor(const TArra
 {
 	EToolkitMode::Type Mode = EditWithinLevelEditor.IsValid() ? EToolkitMode::WorldCentric : EToolkitMode::Standalone;
 
+	for (UObject* Object : InObjects)
+	{
+		UBlueprint* Blueprint = Cast<UBlueprint>(Object);
+		if (Blueprint && Blueprint->SkeletonGeneratedClass && Blueprint->GeneratedClass)
+		{
+			TSharedRef< FWidgetBlueprintEditor > NewBlueprintEditor(new FWidgetBlueprintEditor());
+
+			TArray<UBlueprint*> Blueprints;
+			Blueprints.Add(Blueprint);
+			NewBlueprintEditor->InitWidgetBlueprintEditor(EToolkitMode::Standalone, nullptr, Blueprints, true);
+		}
+		else
+		{
+			FMessageDialog::Open(EAppMsgType::Ok, LOCTEXT("FailedToLoadEditorUtilityWidgetBlueprint", "Editor Utility Widget could not be loaded because it derives from an invalid class.\nCheck to make sure the parent class for this blueprint hasn't been removed!"));
+		}
+	}
+}
+
+uint32 FAssetTypeActions_EditorUtilityWidgetBlueprint::GetCategories()
+{
+	IBlutilityModule* BlutilityModule = FModuleManager::GetModulePtr<IBlutilityModule>("Blutility");
+	return BlutilityModule->GetAssetCategory();
+}
+
+void FAssetTypeActions_EditorUtilityWidgetBlueprint::ExecuteRun(FWeakBlueprintPointerArray InObjects)
+{
 	for (auto ObjIt = InObjects.CreateConstIterator(); ObjIt; ++ObjIt)
 	{
 		if (UWidgetBlueprint* Blueprint = Cast<UWidgetBlueprint>(*ObjIt))
@@ -79,7 +105,7 @@ void FAssetTypeActions_EditorUtilityWidgetBlueprint::OpenAssetEditor(const TArra
 				}
 				else
 				{
-					FName RegistrationName = FName(*CDO->GetPathName());
+					FName RegistrationName = FName(*(CDO->GetPathName() + LOCTEXT("ActiveTabSuffix", "_ActiveTab").ToString()));
 					FText DisplayName = FText::FromString(Blueprint->GetName());
 					FLevelEditorModule& LevelEditorModule = FModuleManager::GetModuleChecked<FLevelEditorModule>(TEXT("LevelEditor"));
 					TSharedPtr<FTabManager> LevelEditorTabManager = LevelEditorModule.GetLevelEditorTabManager();
@@ -94,37 +120,8 @@ void FAssetTypeActions_EditorUtilityWidgetBlueprint::OpenAssetEditor(const TArra
 						BlutilityModule->AddLoadedScriptUI(WidgetBlueprint);
 					}
 					TSharedRef<SDockTab> NewDockTab = LevelEditorTabManager->InvokeTab(RegistrationName);
-				
-				}
-			}
-		}
-	}
-}
-
-uint32 FAssetTypeActions_EditorUtilityWidgetBlueprint::GetCategories()
-{
-	IBlutilityModule* BlutilityModule = FModuleManager::GetModulePtr<IBlutilityModule>("Blutility");
-	return BlutilityModule->GetAssetCategory();
-}
 
-void FAssetTypeActions_EditorUtilityWidgetBlueprint::ExecuteEdit(FWeakBlueprintPointerArray Objects)
-{
-	for (auto ObjIt = Objects.CreateConstIterator(); ObjIt; ++ObjIt)
-	{
-		if (auto Object = (*ObjIt).Get())
-		{
-			auto Blueprint = Cast<UBlueprint>(*ObjIt);
-			if (Blueprint && Blueprint->SkeletonGeneratedClass && Blueprint->GeneratedClass)
-			{
-				TSharedRef< FWidgetBlueprintEditor > NewBlueprintEditor(new FWidgetBlueprintEditor());
-
-				TArray<UBlueprint*> Blueprints;
-				Blueprints.Add(Blueprint);
-				NewBlueprintEditor->InitWidgetBlueprintEditor(EToolkitMode::Standalone, nullptr, Blueprints, true);
-			}
-			else
-			{
-				FMessageDialog::Open(EAppMsgType::Ok, LOCTEXT("FailedToLoadEditorUtilityWidgetBlueprint", "Editor Utility Widget could not be loaded because it derives from an invalid class.\nCheck to make sure the parent class for this blueprint hasn't been removed!"));
+				}
 			}
 		}
 	}
diff --git a/Engine/Source/Editor/Blutility/Private/AssetTypeActions_EditorUtilityWidgetBlueprint.h b/Engine/Source/Editor/Blutility/Private/AssetTypeActions_EditorUtilityWidgetBlueprint.h
--- a/Engine/Source/Editor/Blutility/Private/AssetTypeActions_EditorUtilityWidgetBlueprint.h
+++ b/Engine/Source/Editor/Blutility/Private/AssetTypeActions_EditorUtilityWidgetBlueprint.h
@@ -26,6 +26,5 @@ public:
 protected:
 	typedef TArray< TWeakObjectPtr<class UWidgetBlueprint> > FWeakBlueprintPointerArray;
 
-	void ExecuteEdit(FWeakBlueprintPointerArray Objects);
-
+	void ExecuteRun(FWeakBlueprintPointerArray Objects);
 };
