Fix for editor utility widgets staying in the windows menu after close

diff --git a/Engine/Source/Editor/Blutility/Classes/EditorUtilityWidget.h b/Engine/Source/Editor/Blutility/Classes/EditorUtilityWidget.h
--- a/Engine/Source/Editor/Blutility/Classes/EditorUtilityWidget.h
+++ b/Engine/Source/Editor/Blutility/Classes/EditorUtilityWidget.h
@@ -22,7 +22,7 @@ class BLUTILITY_API UEditorUtilityWidget : public UUserWidget
 	GENERATED_UCLASS_BODY()
 
 public:
-	// The default action called when the blutility is invoked if bAutoRunDefaultAction=true (it is never called otherwise)
+	// The default action called when the widget is invoked if bAutoRunDefaultAction=true (it is never called otherwise)
 	UFUNCTION(BlueprintImplementableEvent)
 	void OnDefaultActionClicked();
 
diff --git a/Engine/Source/Editor/Blutility/Classes/EditorUtilityWidgetBlueprint.h b/Engine/Source/Editor/Blutility/Classes/EditorUtilityWidgetBlueprint.h
--- a/Engine/Source/Editor/Blutility/Classes/EditorUtilityWidgetBlueprint.h
+++ b/Engine/Source/Editor/Blutility/Classes/EditorUtilityWidgetBlueprint.h
@@ -22,6 +22,8 @@ class UEditorUtilityWidgetBlueprint : public UWidgetBlueprint
 	GENERATED_UCLASS_BODY()
 
 public:
+	virtual void BeginDestroy() override;
+
 	TSharedRef<SDockTab> SpawnEditorUITab(const FSpawnTabArgs& SpawnTabArgs);
 
 	/** Creates the slate widget from the UMG widget */
@@ -40,8 +42,16 @@ public:
 		return CreatedUMGWidget;
 	}
 
+	void SetRegistrationName(FName InRegistrationName)
+	{
+		RegistrationName = InRegistrationName;
+	}
+
 private:
+	FName RegistrationName;
+
 	TWeakPtr<SDockTab> CreatedTab;
 
 	UEditorUtilityWidget* CreatedUMGWidget;
+
 };
diff --git a/Engine/Source/Editor/Blutility/Private/AssetTypeActions_EditorUtilityWidgetBlueprint.cpp b/Engine/Source/Editor/Blutility/Private/AssetTypeActions_EditorUtilityWidgetBlueprint.cpp
--- a/Engine/Source/Editor/Blutility/Private/AssetTypeActions_EditorUtilityWidgetBlueprint.cpp
+++ b/Engine/Source/Editor/Blutility/Private/AssetTypeActions_EditorUtilityWidgetBlueprint.cpp
@@ -113,7 +113,7 @@ void FAssetTypeActions_EditorUtilityWidgetBlueprint::ExecuteRun(FWeakBlueprintPo
 					{
 						IBlutilityModule* BlutilityModule = FModuleManager::GetModulePtr<IBlutilityModule>("Blutility");
 						UEditorUtilityWidgetBlueprint* WidgetBlueprint = Cast<UEditorUtilityWidgetBlueprint>(Blueprint);
-
+						WidgetBlueprint->SetRegistrationName(RegistrationName);
 						LevelEditorTabManager->RegisterTabSpawner(RegistrationName, FOnSpawnTab::CreateUObject(WidgetBlueprint, &UEditorUtilityWidgetBlueprint::SpawnEditorUITab))
 							.SetDisplayName(DisplayName)
 							.SetGroup(BlutilityModule->GetMenuGroup().ToSharedRef());
diff --git a/Engine/Source/Editor/Blutility/Private/BlutilityModule.cpp b/Engine/Source/Editor/Blutility/Private/BlutilityModule.cpp
--- a/Engine/Source/Editor/Blutility/Private/BlutilityModule.cpp
+++ b/Engine/Source/Editor/Blutility/Private/BlutilityModule.cpp
@@ -93,35 +93,40 @@ public:
 		FLevelEditorModule& LevelEditorModule = FModuleManager::LoadModuleChecked<FLevelEditorModule>(TEXT("LevelEditor"));
 		LevelEditorModule.OnTabManagerChanged().AddRaw(this, &FBlutilityModule::ReinitializeUIs);
 		LevelEditorModule.OnMapChanged().AddRaw(this, &FBlutilityModule::OnMapChanged);
-
 		FEditorSupportDelegates::PrepareToCleanseEditorObject.AddRaw(this, &FBlutilityModule::OnPrepareToCleanseEditorObject);
 	}
 
 	void ReinitializeUIs()
 	{
-		EditorUtilityContext = NewObject<UEditorUtilityContext>();
-		if (EditorUtilityContext)
+		if (!EditorUtilityContext)
 		{
-			for (FSoftObjectPath BlueprintPath : EditorUtilityContext->LoadedUIs)
-			{ 
-				UObject* BlueprintObject = BlueprintPath.TryLoad();
-				if (BlueprintObject)
+			EditorUtilityContext = NewObject<UEditorUtilityContext>();
+		}
+		FLevelEditorModule& LevelEditorModule = FModuleManager::GetModuleChecked<FLevelEditorModule>(TEXT("LevelEditor"));
+		TSharedPtr<FTabManager> LevelEditorTabManager = LevelEditorModule.GetLevelEditorTabManager();
+		TArray<FSoftObjectPath> CorrectPaths;
+		for (FSoftObjectPath BlueprintPath : EditorUtilityContext->LoadedUIs)
+		{
+			UObject* BlueprintObject = BlueprintPath.TryLoad();
+			if (BlueprintObject && !BlueprintObject->IsPendingKillOrUnreachable())
+			{
+				UEditorUtilityWidgetBlueprint* Blueprint = Cast<UEditorUtilityWidgetBlueprint>(BlueprintObject);
+				const UEditorUtilityWidget* CDO = Blueprint->GeneratedClass->GetDefaultObject<UEditorUtilityWidget>();
+				FName RegistrationName = FName(*(CDO->GetPathName() + LOCTEXT("ActiveTabSuffix", "_ActiveTab").ToString()));
+				Blueprint->SetRegistrationName(RegistrationName);
+				FText DisplayName = FText::FromString(Blueprint->GetName());
+				if (LevelEditorTabManager && !LevelEditorTabManager->CanSpawnTab(RegistrationName))
 				{
-					UEditorUtilityWidgetBlueprint* Blueprint = Cast<UEditorUtilityWidgetBlueprint>(BlueprintObject);
-					const UEditorUtilityWidget* CDO = Blueprint->GeneratedClass->GetDefaultObject<UEditorUtilityWidget>();
-					FName RegistrationName = FName(*CDO->GetPathName());
-					FText DisplayName = FText::FromString(Blueprint->GetName());
-					FLevelEditorModule& LevelEditorModule = FModuleManager::GetModuleChecked<FLevelEditorModule>(TEXT("LevelEditor"));
-					TSharedPtr<FTabManager> LevelEditorTabManager = LevelEditorModule.GetLevelEditorTabManager();
-					if (LevelEditorTabManager && !LevelEditorTabManager->CanSpawnTab(RegistrationName))
-					{
-						LevelEditorTabManager->RegisterTabSpawner(RegistrationName, FOnSpawnTab::CreateUObject(Blueprint, &UEditorUtilityWidgetBlueprint::SpawnEditorUITab))
-							.SetDisplayName(DisplayName)
-							.SetGroup(GetMenuGroup().ToSharedRef());
-					}
+					LevelEditorTabManager->RegisterTabSpawner(RegistrationName, FOnSpawnTab::CreateUObject(Blueprint, &UEditorUtilityWidgetBlueprint::SpawnEditorUITab))
+						.SetDisplayName(DisplayName)
+						.SetGroup(GetMenuGroup().ToSharedRef());
+					CorrectPaths.Add(BlueprintPath);
 				}
 			}
 		}
+
+		EditorUtilityContext->LoadedUIs = CorrectPaths;
+		EditorUtilityContext->SaveConfig();
 	}
 
 	void OnMapChanged(UWorld* InWorld, EMapChangeType MapChangeType)
@@ -217,7 +222,7 @@ public:
 	{
 		if (EditorUtilityContext)
 		{
-			EditorUtilityContext->LoadedUIs.Add(InBlueprint);
+			EditorUtilityContext->LoadedUIs.AddUnique(InBlueprint);
 			EditorUtilityContext->SaveConfig();
 		}
 	}
@@ -231,7 +236,6 @@ public:
 			EditorUtilityContext->SaveConfig();
 		}
 	}
-
 protected:
 	static TSharedRef<SDockTab> SpawnBlutilityShelfTab(const FSpawnTabArgs& Args)
 	{
diff --git a/Engine/Source/Editor/Blutility/Private/EditorUtilityWidgetBlueprint.cpp b/Engine/Source/Editor/Blutility/Private/EditorUtilityWidgetBlueprint.cpp
--- a/Engine/Source/Editor/Blutility/Private/EditorUtilityWidgetBlueprint.cpp
+++ b/Engine/Source/Editor/Blutility/Private/EditorUtilityWidgetBlueprint.cpp
@@ -6,6 +6,7 @@
 #include "EditorUtilityWidget.h"
 #include "IBlutilityModule.h"
 #include "Modules/ModuleManager.h"
+#include "LevelEditor.h"
 
 
 
@@ -18,6 +19,22 @@ UEditorUtilityWidgetBlueprint::UEditorUtilityWidgetBlueprint(const FObjectInitia
 {
 }
 
+void UEditorUtilityWidgetBlueprint::BeginDestroy()
+{
+	if (CreatedTab.IsValid())
+	{
+		IBlutilityModule* BlutilityModule = FModuleManager::GetModulePtr<IBlutilityModule>("Blutility");
+		BlutilityModule->RemoveLoadedScriptUI(this);
+
+		FLevelEditorModule& LevelEditorModule = FModuleManager::GetModuleChecked<FLevelEditorModule>(TEXT("LevelEditor"));
+		TSharedPtr<FTabManager> LevelEditorTabManager = LevelEditorModule.GetLevelEditorTabManager();
+		LevelEditorTabManager->UnregisterTabSpawner(RegistrationName);
+	}
+
+	Super::BeginDestroy();
+}
+
+
 TSharedRef<SDockTab> UEditorUtilityWidgetBlueprint::SpawnEditorUITab(const FSpawnTabArgs& SpawnTabArgs)
 {
 	TSharedRef<SDockTab> SpawnedTab = SNew(SDockTab);
