Adding code to close all editor widget tabs if we are closing all editors

diff --git a/Engine/Source/Editor/Blutility/Classes/EditorUtilityWidgetBlueprint.h b/Engine/Source/Editor/Blutility/Classes/EditorUtilityWidgetBlueprint.h
--- a/Engine/Source/Editor/Blutility/Classes/EditorUtilityWidgetBlueprint.h
+++ b/Engine/Source/Editor/Blutility/Classes/EditorUtilityWidgetBlueprint.h
@@ -15,6 +15,7 @@
 
 class UBlueprint;
 class UEditorUtilityWidget;
+enum class EAssetEditorCloseReason : uint8;
 
 UCLASS()
 class UEditorUtilityWidgetBlueprint : public UWidgetBlueprint
@@ -47,6 +48,11 @@ public:
 		RegistrationName = InRegistrationName;
 	}
 
+	FName GetRegistrationName() const
+	{
+		return RegistrationName;
+	}
+
 private:
 	FName RegistrationName;
 
diff --git a/Engine/Source/Editor/Blutility/Private/AssetTypeActions_EditorUtilityWidgetBlueprint.cpp b/Engine/Source/Editor/Blutility/Private/AssetTypeActions_EditorUtilityWidgetBlueprint.cpp
--- a/Engine/Source/Editor/Blutility/Private/AssetTypeActions_EditorUtilityWidgetBlueprint.cpp
+++ b/Engine/Source/Editor/Blutility/Private/AssetTypeActions_EditorUtilityWidgetBlueprint.cpp
@@ -105,7 +105,7 @@ void FAssetTypeActions_EditorUtilityWidgetBlueprint::ExecuteRun(FWeakBlueprintPo
 				}
 				else
 				{
-					FName RegistrationName = FName(*(CDO->GetPathName() + LOCTEXT("ActiveTabSuffix", "_ActiveTab").ToString()));
+					FName RegistrationName = FName(*(Blueprint->GetPathName() + LOCTEXT("ActiveTabSuffix", "_ActiveTab").ToString()));
 					FText DisplayName = FText::FromString(Blueprint->GetName());
 					FLevelEditorModule& LevelEditorModule = FModuleManager::GetModuleChecked<FLevelEditorModule>(TEXT("LevelEditor"));
 					TSharedPtr<FTabManager> LevelEditorTabManager = LevelEditorModule.GetLevelEditorTabManager();
diff --git a/Engine/Source/Editor/Blutility/Private/BlutilityModule.cpp b/Engine/Source/Editor/Blutility/Private/BlutilityModule.cpp
--- a/Engine/Source/Editor/Blutility/Private/BlutilityModule.cpp
+++ b/Engine/Source/Editor/Blutility/Private/BlutilityModule.cpp
@@ -33,6 +33,7 @@
 #include "UnrealEdMisc.h"
 #include "EditorSupportDelegates.h"
 #include "UObject/PurgingReferenceCollector.h"
+#include "AssetRegistryModule.h"
 
 #define LOCTEXT_NAMESPACE "AssetTypeActions"
 
@@ -94,6 +95,9 @@ public:
 		LevelEditorModule.OnTabManagerChanged().AddRaw(this, &FBlutilityModule::ReinitializeUIs);
 		LevelEditorModule.OnMapChanged().AddRaw(this, &FBlutilityModule::OnMapChanged);
 		FEditorSupportDelegates::PrepareToCleanseEditorObject.AddRaw(this, &FBlutilityModule::OnPrepareToCleanseEditorObject);
+
+		FAssetRegistryModule& AssetRegistryModule = FModuleManager::LoadModuleChecked<FAssetRegistryModule>(TEXT("AssetRegistry"));
+		AssetRegistryModule.Get().OnAssetRemoved().AddRaw(this, &FBlutilityModule::HandleAssetRemoved);
 	}
 
 	void ReinitializeUIs()
@@ -112,7 +116,7 @@ public:
 			{
 				UEditorUtilityWidgetBlueprint* Blueprint = Cast<UEditorUtilityWidgetBlueprint>(BlueprintObject);
 				const UEditorUtilityWidget* CDO = Blueprint->GeneratedClass->GetDefaultObject<UEditorUtilityWidget>();
-				FName RegistrationName = FName(*(CDO->GetPathName() + LOCTEXT("ActiveTabSuffix", "_ActiveTab").ToString()));
+				FName RegistrationName = FName(*(Blueprint->GetPathName() + LOCTEXT("ActiveTabSuffix", "_ActiveTab").ToString()));
 				Blueprint->SetRegistrationName(RegistrationName);
 				FText DisplayName = FText::FromString(Blueprint->GetName());
 				if (LevelEditorTabManager && !LevelEditorTabManager->CanSpawnTab(RegistrationName))
@@ -236,6 +240,7 @@ public:
 			EditorUtilityContext->SaveConfig();
 		}
 	}
+
 protected:
 	static TSharedRef<SDockTab> SpawnBlutilityShelfTab(const FSpawnTabArgs& Args)
 	{
@@ -289,6 +294,34 @@ protected:
 		}
 	}
 
+	void HandleAssetRemoved(const FAssetData& InAssetData)
+	{
+		if (EditorUtilityContext)
+		{
+			bool bDeletingLoadedUI = false;
+			for (FSoftObjectPath LoadedUIPath : EditorUtilityContext->LoadedUIs)
+			{
+				if (LoadedUIPath.GetAssetPathName() == InAssetData.ObjectPath)
+				{
+					bDeletingLoadedUI = true;
+					break;
+				}
+			}
+
+			if (bDeletingLoadedUI)
+			{
+				FName UIToCleanup = FName(*(InAssetData.ObjectPath.ToString() + LOCTEXT("ActiveTabSuffix", "_ActiveTab").ToString()));
+				FLevelEditorModule& LevelEditorModule = FModuleManager::GetModuleChecked<FLevelEditorModule>(TEXT("LevelEditor"));
+				TSharedPtr<FTabManager> LevelEditorTabManager = LevelEditorModule.GetLevelEditorTabManager();
+				TSharedPtr<SDockTab> CurrentTab = LevelEditorTabManager->FindExistingLiveTab(UIToCleanup);
+				if (CurrentTab.IsValid())
+				{
+					CurrentTab->RequestCloseTab();
+				}
+			}
+		}
+	}
+protected:
 	/** Scripted Editor Widgets workspace menu item */
 	TSharedPtr<class FWorkspaceItem> ScriptedEditorWidgetsGroup;
 
diff --git a/Engine/Source/Editor/Blutility/Private/EditorUtilityWidgetBlueprint.cpp b/Engine/Source/Editor/Blutility/Private/EditorUtilityWidgetBlueprint.cpp
--- a/Engine/Source/Editor/Blutility/Private/EditorUtilityWidgetBlueprint.cpp
+++ b/Engine/Source/Editor/Blutility/Private/EditorUtilityWidgetBlueprint.cpp
@@ -17,10 +17,12 @@
 UEditorUtilityWidgetBlueprint::UEditorUtilityWidgetBlueprint(const FObjectInitializer& ObjectInitializer)
 	: Super(ObjectInitializer)
 {
+
 }
 
 void UEditorUtilityWidgetBlueprint::BeginDestroy()
 {
+	// Created tab is invalid on editor shutdown, prevents the cleanup script from running then
 	if (CreatedTab.IsValid())
 	{
 		IBlutilityModule* BlutilityModule = FModuleManager::GetModulePtr<IBlutilityModule>("Blutility");
