Override IsAssetEditor to prevent FGlobalEditorUtilityDialog from showing up as a result of FindEditorForAsset for blueprint assets

diff --git a/Engine/Source/Editor/Blutility/Private/EditorUtilityDialog.h b/Engine/Source/Editor/Blutility/Private/EditorUtilityDialog.h
--- a/Engine/Source/Editor/Blutility/Private/EditorUtilityDialog.h
+++ b/Engine/Source/Editor/Blutility/Private/EditorUtilityDialog.h
@@ -44,6 +44,7 @@ public:
 	virtual FName GetToolkitFName() const override;
 	virtual FText GetBaseToolkitName() const override;
 	virtual FString GetWorldCentricTabPrefix() const override;
+	virtual bool IsAssetEditor() const override { return false; }
 
 	/** @return Returns the color and opacity to use for the color that appears behind the tab text for this toolkit's tab in world-centric mode. */
 	virtual FLinearColor GetWorldCentricTabColorScale() const override;
