Fixing case in which Editor Utility Widget Blueprint does its cleanup

diff --git a/Engine/Source/Editor/Blutility/Private/EditorUtilityWidgetBlueprint.cpp b/Engine/Source/Editor/Blutility/Private/EditorUtilityWidgetBlueprint.cpp
--- a/Engine/Source/Editor/Blutility/Private/EditorUtilityWidgetBlueprint.cpp
+++ b/Engine/Source/Editor/Blutility/Private/EditorUtilityWidgetBlueprint.cpp
@@ -22,8 +22,8 @@ UEditorUtilityWidgetBlueprint::UEditorUtilityWidgetBlueprint(const FObjectInitia
 
 void UEditorUtilityWidgetBlueprint::BeginDestroy()
 {
-	// Created tab is invalid on editor shutdown, prevents the cleanup script from running then
-	if (CreatedTab.IsValid())
+	// prevent the cleanup script from running on editor shutdown
+	if (!GIsRequestingExit)
 	{
 		IBlutilityModule* BlutilityModule = FModuleManager::GetModulePtr<IBlutilityModule>("Blutility");
 		BlutilityModule->RemoveLoadedScriptUI(this);
