Adding editor-only camera for scene captures
#rnx

diff --git a/Engine/Source/Editor/Blutility/Classes/EditorUtilityCamera.h b/Engine/Source/Editor/Blutility/Classes/EditorUtilityCamera.h
new file mode 100644
--- /dev/null
+++ b/Engine/Source/Editor/Blutility/Classes/EditorUtilityCamera.h
@@ -0,0 +1,19 @@
+/**
+ * Base class of any editor-only camera actors. Should be used when creating editor-only capture tools, etc.
+ */
+
+#pragma once
+
+#include "CoreMinimal.h"
+#include "Camera/CameraActor.h"
+#include "UObject/ScriptMacros.h"
+#include "EditorUtilityCamera.generated.h"
+
+
+UCLASS(Abstract, Blueprintable)
+class BLUTILITY_API AEditorUtilityCamera : public ACameraActor
+{
+	GENERATED_UCLASS_BODY()
+
+
+};
\ No newline at end of file
diff --git a/Engine/Source/Editor/Blutility/Private/EditorUtilityCamera.cpp b/Engine/Source/Editor/Blutility/Private/EditorUtilityCamera.cpp
new file mode 100644
--- /dev/null
+++ b/Engine/Source/Editor/Blutility/Private/EditorUtilityCamera.cpp
@@ -0,0 +1,11 @@
+// Copyright 1998-2019 Epic Games, Inc. All Rights Reserved.
+
+#include "EditorUtilityCamera.h"
+
+
+/////////////////////////////////////////////////////
+
+AEditorUtilityCamera::AEditorUtilityCamera(const FObjectInitializer& ObjectInitializer)
+	: Super(ObjectInitializer)
+{
+}
