Fixing header notice
#rnx

diff --git a/Engine/Source/Editor/Blutility/Classes/EditorUtilityCamera.h b/Engine/Source/Editor/Blutility/Classes/EditorUtilityCamera.h
--- a/Engine/Source/Editor/Blutility/Classes/EditorUtilityCamera.h
+++ b/Engine/Source/Editor/Blutility/Classes/EditorUtilityCamera.h
@@ -1,3 +1,5 @@
+// Copyright 1998-2019 Epic Games, Inc. All Rights Reserved.
+
 /**
  * Base class of any editor-only camera actors. Should be used when creating editor-only capture tools, etc.
  */
@@ -10,7 +12,7 @@
 #include "EditorUtilityCamera.generated.h"
 
 
-UCLASS(Abstract, Blueprintable)
+UCLASS(Abstract, Blueprintable, meta = (ShowWorldContextPin))
 class BLUTILITY_API AEditorUtilityCamera : public ACameraActor
 {
 	GENERATED_UCLASS_BODY()
