Fix for a crash when reopening level after closing UI
#rnx

diff --git a/Engine/Source/Editor/Blutility/Private/EditorUtilityWidgetBlueprint.cpp b/Engine/Source/Editor/Blutility/Private/EditorUtilityWidgetBlueprint.cpp
--- a/Engine/Source/Editor/Blutility/Private/EditorUtilityWidgetBlueprint.cpp
+++ b/Engine/Source/Editor/Blutility/Private/EditorUtilityWidgetBlueprint.cpp
@@ -99,6 +99,7 @@ void UEditorUtilityWidgetBlueprint::UpdateRespawnListIfNeeded(TSharedRef<SDockTa
 		IBlutilityModule* BlutilityModule = FModuleManager::GetModulePtr<IBlutilityModule>("Blutility");
 		BlutilityModule->RemoveLoadedScriptUI(this);
 	}
+	CreatedUMGWidget = nullptr;
 }
 
 void UEditorUtilityWidgetBlueprint::GetReparentingRules(TSet< const UClass* >& AllowedChildrenOfClasses, TSet< const UClass* >& DisallowedChildrenOfClasses) const
