Fix to correctly null out widget references when closing editor widget tabs

diff --git a/Engine/Source/Editor/Blutility/Private/EditorUtilityWidgetBlueprint.cpp b/Engine/Source/Editor/Blutility/Private/EditorUtilityWidgetBlueprint.cpp
--- a/Engine/Source/Editor/Blutility/Private/EditorUtilityWidgetBlueprint.cpp
+++ b/Engine/Source/Editor/Blutility/Private/EditorUtilityWidgetBlueprint.cpp
@@ -93,7 +93,7 @@ void UEditorUtilityWidgetBlueprint::RegenerateCreatedTab(UBlueprint* RecompiledB
 
 void UEditorUtilityWidgetBlueprint::UpdateRespawnListIfNeeded(TSharedRef<SDockTab> TabBeingClosed)
 {
-	UEditorUtilityWidget* EditorUtilityWidget = Cast<UEditorUtilityWidget>(GeneratedClass);
+	const UEditorUtilityWidget* EditorUtilityWidget = GeneratedClass->GetDefaultObject<UEditorUtilityWidget>();
 	if (EditorUtilityWidget && EditorUtilityWidget->ShouldAlwaysReregisterWithWindowsMenu() == false)
 	{
 		IBlutilityModule* BlutilityModule = FModuleManager::GetModulePtr<IBlutilityModule>("Blutility");
