Fixing Editor Utility filter to not filter if no allowed or disallowed classes are set
#rnx

diff --git a/Engine/Source/Editor/Blutility/Private/EditorUtilityBlueprintFactory.cpp b/Engine/Source/Editor/Blutility/Private/EditorUtilityBlueprintFactory.cpp
--- a/Engine/Source/Editor/Blutility/Private/EditorUtilityBlueprintFactory.cpp
+++ b/Engine/Source/Editor/Blutility/Private/EditorUtilityBlueprintFactory.cpp
@@ -25,12 +25,21 @@ public:
 
 	bool IsClassAllowed(const FClassViewerInitializationOptions& InInitOptions, const UClass* InClass, TSharedRef< FClassViewerFilterFuncs > InFilterFuncs ) override
 	{
+		if (DisallowedChildOfClasses.Num() == 0 && AllowedChildOfClasses.Num() == 0)
+		{
+			return true;
+		}
 		return (InFilterFuncs->IfInChildOfClassesSet(AllowedChildOfClasses, InClass) != EFilterReturn::Failed) 
 			&& (InFilterFuncs->IfInChildOfClassesSet(DisallowedChildOfClasses, InClass) == EFilterReturn::Failed);
 	}
 
 	virtual bool IsUnloadedClassAllowed(const FClassViewerInitializationOptions& InInitOptions, const TSharedRef< const IUnloadedBlueprintData > InUnloadedClassData, TSharedRef< FClassViewerFilterFuncs > InFilterFuncs) override
 	{
+		if (DisallowedChildOfClasses.Num() == 0 && AllowedChildOfClasses.Num() == 0)
+		{
+			return true;
+		}
+
 		return InFilterFuncs->IfInChildOfClassesSet(AllowedChildOfClasses, InUnloadedClassData) != EFilterReturn::Failed
 			&& (InFilterFuncs->IfInChildOfClassesSet(DisallowedChildOfClasses, InUnloadedClassData) == EFilterReturn::Failed);;
 	}
