Run editor blueprints at startup based on config file

diff --git a/Engine/Source/Editor/Blutility/Private/BlutilityModule.cpp b/Engine/Source/Editor/Blutility/Private/BlutilityModule.cpp
--- a/Engine/Source/Editor/Blutility/Private/BlutilityModule.cpp
+++ b/Engine/Source/Editor/Blutility/Private/BlutilityModule.cpp
@@ -34,9 +34,12 @@
 #include "EditorSupportDelegates.h"
 #include "UObject/PurgingReferenceCollector.h"
 #include "AssetRegistryModule.h"
+#include "EditorUtilityCommon.h"
 
 #define LOCTEXT_NAMESPACE "AssetTypeActions"
 
+DEFINE_LOG_CATEGORY(LogEditorUtilityBlueprint);
+
 /////////////////////////////////////////////////////
 
 namespace BlutilityModule
@@ -103,27 +106,85 @@ public:
 		FLevelEditorModule& LevelEditorModule = FModuleManager::GetModuleChecked<FLevelEditorModule>(TEXT("LevelEditor"));
 		TSharedPtr<FTabManager> LevelEditorTabManager = LevelEditorModule.GetLevelEditorTabManager();
 		TArray<FSoftObjectPath> CorrectPaths;
-		for (FSoftObjectPath BlueprintPath : EditorUtilityContext->LoadedUIs)
+		for (FSoftObjectPath& BlueprintPath : EditorUtilityContext->LoadedUIs)
 		{
 			UObject* BlueprintObject = BlueprintPath.TryLoad();
 			if (BlueprintObject && !BlueprintObject->IsPendingKillOrUnreachable())
 			{
 				UEditorUtilityWidgetBlueprint* Blueprint = Cast<UEditorUtilityWidgetBlueprint>(BlueprintObject);
-				const UEditorUtilityWidget* CDO = Blueprint->GeneratedClass->GetDefaultObject<UEditorUtilityWidget>();
-				FName RegistrationName = FName(*(Blueprint->GetPathName() + LOCTEXT("ActiveTabSuffix", "_ActiveTab").ToString()));
-				Blueprint->SetRegistrationName(RegistrationName);
-				FText DisplayName = FText::FromString(Blueprint->GetName());
-				if (LevelEditorTabManager && !LevelEditorTabManager->CanSpawnTab(RegistrationName))
+				if (Blueprint)
 				{
-					LevelEditorTabManager->RegisterTabSpawner(RegistrationName, FOnSpawnTab::CreateUObject(Blueprint, &UEditorUtilityWidgetBlueprint::SpawnEditorUITab))
-						.SetDisplayName(DisplayName)
-						.SetGroup(GetMenuGroup().ToSharedRef());
-					CorrectPaths.Add(BlueprintPath);
+					if (Blueprint->GeneratedClass)
+					{
+						const UEditorUtilityWidget* CDO = Blueprint->GeneratedClass->GetDefaultObject<UEditorUtilityWidget>();
+						FName RegistrationName = FName(*(Blueprint->GetPathName() + LOCTEXT("ActiveTabSuffix", "_ActiveTab").ToString()));
+						Blueprint->SetRegistrationName(RegistrationName);
+						FText DisplayName = FText::FromString(Blueprint->GetName());
+						if (LevelEditorTabManager && !LevelEditorTabManager->CanSpawnTab(RegistrationName))
+						{
+							LevelEditorTabManager->RegisterTabSpawner(RegistrationName, FOnSpawnTab::CreateUObject(Blueprint, &UEditorUtilityWidgetBlueprint::SpawnEditorUITab))
+								.SetDisplayName(DisplayName)
+								.SetGroup(GetMenuGroup().ToSharedRef());
+							CorrectPaths.Add(BlueprintPath);
+						}
+					}
+					else
+					{
+						UE_LOG(LogEditorUtilityBlueprint, Warning, TEXT("No generated class for: %s"), *BlueprintPath.ToString());
+					}
 				}
+				else
+				{
+					UE_LOG(LogEditorUtilityBlueprint, Warning, TEXT("Expected object of class EditorUtilityWidgetBlueprint: %s"), *BlueprintPath.ToString());
+				}
+			}
+			else
+			{
+				UE_LOG(LogEditorUtilityBlueprint, Warning, TEXT("Could not load: %s"), *BlueprintPath.ToString());
 			}
 		}
 
 		EditorUtilityContext->LoadedUIs = CorrectPaths;
+
+		for (FSoftObjectPath& ObjectPath : EditorUtilityContext->StartupBlueprints)
+		{
+			UObject* Object = ObjectPath.TryLoad();
+			if (!Object || Object->IsPendingKillOrUnreachable())
+			{
+				UE_LOG(LogEditorUtilityBlueprint, Warning, TEXT("Could not load: %s"), *ObjectPath.ToString());
+				continue;
+			}
+
+			UClass* ObjectClass = Object->GetClass();
+			if (UBlueprint* Blueprint = Cast<UBlueprint>(Object))
+			{
+				ObjectClass = Blueprint->GeneratedClass;
+				if (!ObjectClass)
+				{
+					UE_LOG(LogEditorUtilityBlueprint, Warning, TEXT("Editor startup blueprint class could not be generated: %s"), *ObjectPath.ToString());
+					continue;
+				}
+			}
+
+			if (ObjectClass)
+			{
+				static const FName RunFunctionName("Run");
+				UFunction* RunFunction = ObjectClass->FindFunctionByName(RunFunctionName);
+				if (RunFunction)
+				{
+					UObject* TempObject = NewObject<UObject>(GetTransientPackage(), ObjectClass);
+					TempObject->AddToRoot();
+					Object->ProcessEvent(RunFunction, nullptr);
+					TempObject->RemoveFromRoot();
+				}
+				else
+				{
+					UE_LOG(LogEditorUtilityBlueprint, Warning, TEXT("Editor startup blueprint missing function named 'Run': %s"), *ObjectPath.ToString());
+					continue;
+				}
+			}
+		}
+
 		EditorUtilityContext->SaveConfig();
 	}
 
diff --git a/Engine/Source/Editor/Blutility/Private/EditorUtilityCommon.h b/Engine/Source/Editor/Blutility/Private/EditorUtilityCommon.h
new file mode 100644
--- /dev/null
+++ b/Engine/Source/Editor/Blutility/Private/EditorUtilityCommon.h
@@ -0,0 +1,7 @@
+// Copyright 1998-2019 Epic Games, Inc. All Rights Reserved.
+
+#pragma once
+
+#include "CoreMinimal.h"
+
+DECLARE_LOG_CATEGORY_EXTERN(LogEditorUtilityBlueprint, Log, All);
diff --git a/Engine/Source/Editor/Blutility/Private/EditorUtilityContext.h b/Engine/Source/Editor/Blutility/Private/EditorUtilityContext.h
--- a/Engine/Source/Editor/Blutility/Private/EditorUtilityContext.h
+++ b/Engine/Source/Editor/Blutility/Private/EditorUtilityContext.h
@@ -21,4 +21,7 @@ class UEditorUtilityContext : public UObject
 public:
 	UPROPERTY(config)
 	TArray<FSoftObjectPath> LoadedUIs;
+
+	UPROPERTY(config)
+	TArray<FSoftObjectPath> StartupBlueprints;
 };
