Editor Utility Subsystem added to manage startup objects

diff --git a/Engine/Config/BaseEngine.ini b/Engine/Config/BaseEngine.ini
--- a/Engine/Config/BaseEngine.ini
+++ b/Engine/Config/BaseEngine.ini
@@ -656,6 +656,7 @@ bUseStreamingPause=false
 ; 4.23
 +FunctionRedirects=(OldName="SkeletalMeshComponent.SetAnimInstanceClass",NewName="SkeletalMeshComponent.SetAnimClass")
 +FunctionRedirects=(OldName="SkeletalMeshComponent.K2_SetAnimInstanceClass",NewName="SkeletalMeshComponent.SetAnimClass")
++FunctionRedirects=(OldName="EditorUtilityWidget.OnDefaultActionClicked",NewName="EditorUtilityWidget.Run")
 
 +ClassRedirects=(OldName="/Script/CoreUObject.MulticastDelegateProperty",NewName="/Script/CoreUObject.MulticastInlineDelegateProperty")
 +ClassRedirects=(OldName="EditorAutomationActor",NewName="EditorUtilityActor")
diff --git a/Engine/Source/Editor/Blutility/Blutility.Build.cs b/Engine/Source/Editor/Blutility/Blutility.Build.cs
--- a/Engine/Source/Editor/Blutility/Blutility.Build.cs
+++ b/Engine/Source/Editor/Blutility/Blutility.Build.cs
@@ -10,7 +10,10 @@ public class Blutility : ModuleRules
 
 		PrivateIncludePathModuleNames.Add("AssetTools");
 
-        PublicDependencyModuleNames.Add("EditorSubsystem");
+        PublicDependencyModuleNames.AddRange(new string[] {
+			"EditorSubsystem",
+			"MainFrame"
+		});
 
 		PrivateDependencyModuleNames.AddRange(
 			new string[] {
diff --git a/Engine/Source/Editor/Blutility/Classes/EditorUtilityActor.h b/Engine/Source/Editor/Blutility/Classes/EditorUtilityActor.h
--- a/Engine/Source/Editor/Blutility/Classes/EditorUtilityActor.h
+++ b/Engine/Source/Editor/Blutility/Classes/EditorUtilityActor.h
@@ -17,5 +17,7 @@ class BLUTILITY_API AEditorUtilityActor : public AActor
 {
 	GENERATED_UCLASS_BODY()
 
-	
+	// Standard function to execute
+	UFUNCTION(BlueprintImplementableEvent)
+	void Run();
 };
diff --git a/Engine/Source/Editor/Blutility/Classes/EditorUtilityObject.h b/Engine/Source/Editor/Blutility/Classes/EditorUtilityObject.h
--- a/Engine/Source/Editor/Blutility/Classes/EditorUtilityObject.h
+++ b/Engine/Source/Editor/Blutility/Classes/EditorUtilityObject.h
@@ -18,5 +18,7 @@ class BLUTILITY_API UEditorUtilityObject : public UObject
 {
 	GENERATED_UCLASS_BODY()
 
-	
+	// Standard function to execute
+	UFUNCTION(BlueprintImplementableEvent)
+	void Run();
 };
diff --git a/Engine/Source/Editor/Blutility/Classes/EditorUtilityWidget.h b/Engine/Source/Editor/Blutility/Classes/EditorUtilityWidget.h
--- a/Engine/Source/Editor/Blutility/Classes/EditorUtilityWidget.h
+++ b/Engine/Source/Editor/Blutility/Classes/EditorUtilityWidget.h
@@ -24,7 +24,7 @@ class BLUTILITY_API UEditorUtilityWidget : public UUserWidget
 public:
 	// The default action called when the widget is invoked if bAutoRunDefaultAction=true (it is never called otherwise)
 	UFUNCTION(BlueprintImplementableEvent)
-	void OnDefaultActionClicked();
+	void Run();
 
 	// Run the default action
 	void ExecuteDefaultAction();
diff --git a/Engine/Source/Editor/Blutility/Private/AssetTypeActions_EditorUtilityBlueprint.cpp b/Engine/Source/Editor/Blutility/Private/AssetTypeActions_EditorUtilityBlueprint.cpp
--- a/Engine/Source/Editor/Blutility/Private/AssetTypeActions_EditorUtilityBlueprint.cpp
+++ b/Engine/Source/Editor/Blutility/Private/AssetTypeActions_EditorUtilityBlueprint.cpp
@@ -12,6 +12,7 @@
 #include "IContentBrowserSingleton.h"
 #include "ContentBrowserModule.h"
 #include "IBlutilityModule.h"
+#include "EditorUtilitySubsystem.h"
 
 #define LOCTEXT_NAMESPACE "AssetTypeActions"
 
@@ -33,6 +34,25 @@ UClass* FAssetTypeActions_EditorUtilityBlueprint::GetSupportedClass() const
 	return UEditorUtilityBlueprint::StaticClass();
 }
 
+bool FAssetTypeActions_EditorUtilityBlueprint::HasActions(const TArray<UObject*>& InObjects) const
+{
+	return true;
+}
+
+void FAssetTypeActions_EditorUtilityBlueprint::GetActions(const TArray<UObject*>& InObjects, FMenuBuilder& MenuBuilder)
+{
+	auto Blueprints = GetTypedWeakObjectPtrs<UEditorUtilityBlueprint>(InObjects);
+
+	MenuBuilder.AddMenuEntry(
+		LOCTEXT("EditorUtility_Run", "Run Editor Utility Blueprint"),
+		LOCTEXT("EditorUtility_RunTooltip", "Runs this Editor Utility Blueprint."),
+		FSlateIcon(),
+		FUIAction(
+			FExecuteAction::CreateSP(this, &FAssetTypeActions_EditorUtilityBlueprint::ExecuteRun, Blueprints)
+		)
+	);
+}
+
 uint32 FAssetTypeActions_EditorUtilityBlueprint::GetCategories()
 {
 	IBlutilityModule* BlutilityModule = FModuleManager::GetModulePtr<IBlutilityModule>("Blutility");
@@ -67,4 +87,13 @@ void FAssetTypeActions_EditorUtilityBlueprint::ExecuteNewDerivedBlueprint(TWeakO
 	}
 }
 
+void FAssetTypeActions_EditorUtilityBlueprint::ExecuteRun(FWeakBlueprintPointerArray InObjects)
+{
+	UEditorUtilitySubsystem* EditorUtilitySubsystem = GEditor->GetEditorSubsystem<UEditorUtilitySubsystem>();
+	for (auto ObjIt = InObjects.CreateIterator(); ObjIt; ++ObjIt)
+	{
+		EditorUtilitySubsystem->TryRun(ObjIt->Get());
+	}
+}
+
 #undef LOCTEXT_NAMESPACE
diff --git a/Engine/Source/Editor/Blutility/Private/AssetTypeActions_EditorUtilityBlueprint.h b/Engine/Source/Editor/Blutility/Private/AssetTypeActions_EditorUtilityBlueprint.h
--- a/Engine/Source/Editor/Blutility/Private/AssetTypeActions_EditorUtilityBlueprint.h
+++ b/Engine/Source/Editor/Blutility/Private/AssetTypeActions_EditorUtilityBlueprint.h
@@ -16,6 +16,8 @@ public:
 	virtual FText GetName() const override;
 	virtual FColor GetTypeColor() const override;
 	virtual UClass* GetSupportedClass() const override;
+	virtual bool HasActions(const TArray<UObject*>& InObjects) const override;
+	virtual void GetActions(const TArray<UObject*>& InObjects, FMenuBuilder& MenuBuilder) override;
 	virtual uint32 GetCategories() override;
 	virtual bool CanLocalize() const override { return false; }
 	// End of IAssetTypeActions interface
@@ -24,4 +26,5 @@ protected:
 	typedef TArray< TWeakObjectPtr<class UEditorUtilityBlueprint> > FWeakBlueprintPointerArray;
 
 	void ExecuteNewDerivedBlueprint(TWeakObjectPtr<class UEditorUtilityBlueprint> InObject);
+	void ExecuteRun(FWeakBlueprintPointerArray Objects);
 };
diff --git a/Engine/Source/Editor/Blutility/Private/BlutilityModule.cpp b/Engine/Source/Editor/Blutility/Private/BlutilityModule.cpp
--- a/Engine/Source/Editor/Blutility/Private/BlutilityModule.cpp
+++ b/Engine/Source/Editor/Blutility/Private/BlutilityModule.cpp
@@ -27,7 +27,7 @@
 #include "KismetCompilerModule.h"
 #include "WidgetBlueprintCompiler.h"
 #include "UMGEditorModule.h"
-#include "EditorUtilityContext.h"
+#include "EditorUtilitySubsystem.h"
 #include "LevelEditor.h"
 #include "Editor.h"
 #include "UnrealEdMisc.h"
@@ -99,14 +99,11 @@ public:
 
 	void ReinitializeUIs()
 	{
-		if (!EditorUtilityContext)
-		{
-			EditorUtilityContext = NewObject<UEditorUtilityContext>();
-		}
+		UEditorUtilitySubsystem* EditorUtilitySubsystem = GEditor->GetEditorSubsystem<UEditorUtilitySubsystem>();
 		FLevelEditorModule& LevelEditorModule = FModuleManager::GetModuleChecked<FLevelEditorModule>(TEXT("LevelEditor"));
 		TSharedPtr<FTabManager> LevelEditorTabManager = LevelEditorModule.GetLevelEditorTabManager();
 		TArray<FSoftObjectPath> CorrectPaths;
-		for (FSoftObjectPath& BlueprintPath : EditorUtilityContext->LoadedUIs)
+		for (const FSoftObjectPath& BlueprintPath : EditorUtilitySubsystem->LoadedUIs)
 		{
 			UObject* BlueprintObject = BlueprintPath.TryLoad();
 			if (BlueprintObject && !BlueprintObject->IsPendingKillOrUnreachable())
@@ -144,72 +141,29 @@ public:
 			}
 		}
 
-		EditorUtilityContext->LoadedUIs = CorrectPaths;
-
-		for (FSoftObjectPath& ObjectPath : EditorUtilityContext->StartupBlueprints)
-		{
-			UObject* Object = ObjectPath.TryLoad();
-			if (!Object || Object->IsPendingKillOrUnreachable())
-			{
-				UE_LOG(LogEditorUtilityBlueprint, Warning, TEXT("Could not load: %s"), *ObjectPath.ToString());
-				continue;
-			}
-
-			UClass* ObjectClass = Object->GetClass();
-			if (UBlueprint* Blueprint = Cast<UBlueprint>(Object))
-			{
-				ObjectClass = Blueprint->GeneratedClass;
-				if (!ObjectClass)
-				{
-					UE_LOG(LogEditorUtilityBlueprint, Warning, TEXT("Editor startup blueprint class could not be generated: %s"), *ObjectPath.ToString());
-					continue;
-				}
-			}
-
-			if (ObjectClass)
-			{
-				static const FName RunFunctionName("Run");
-				UFunction* RunFunction = ObjectClass->FindFunctionByName(RunFunctionName);
-				if (RunFunction)
-				{
-					UObject* TempObject = NewObject<UObject>(GetTransientPackage(), ObjectClass);
-					TempObject->AddToRoot();
-					Object->ProcessEvent(RunFunction, nullptr);
-					TempObject->RemoveFromRoot();
-				}
-				else
-				{
-					UE_LOG(LogEditorUtilityBlueprint, Warning, TEXT("Editor startup blueprint missing function named 'Run': %s"), *ObjectPath.ToString());
-					continue;
-				}
-			}
-		}
-
-		EditorUtilityContext->SaveConfig();
+		EditorUtilitySubsystem->LoadedUIs = CorrectPaths;
+		EditorUtilitySubsystem->SaveConfig();
 	}
 
 	void OnMapChanged(UWorld* InWorld, EMapChangeType MapChangeType)
 	{
-		if (EditorUtilityContext)
+		for (const FSoftObjectPath& LoadedUI : GEditor->GetEditorSubsystem<UEditorUtilitySubsystem>()->LoadedUIs)
 		{
-			for (FSoftObjectPath LoadedUI : EditorUtilityContext->LoadedUIs)
+			UEditorUtilityWidgetBlueprint* LoadedEditorUtilityBlueprint = Cast<UEditorUtilityWidgetBlueprint>(LoadedUI.ResolveObject());
+			if (LoadedEditorUtilityBlueprint)
 			{
-				UEditorUtilityWidgetBlueprint* LoadedEditorUtilityBlueprint = Cast<UEditorUtilityWidgetBlueprint>(LoadedUI.ResolveObject());
-				if (LoadedEditorUtilityBlueprint)
+				UEditorUtilityWidget* CreatedWidget = LoadedEditorUtilityBlueprint->GetCreatedWidget();
+				if (CreatedWidget)
 				{
-					UEditorUtilityWidget* CreatedWidget = LoadedEditorUtilityBlueprint->GetCreatedWidget();
-					if (CreatedWidget)
+					if (MapChangeType == EMapChangeType::TearDownWorld)
 					{
-						if (MapChangeType == EMapChangeType::TearDownWorld)
-						{
-							CreatedWidget->Rename(*CreatedWidget->GetName(), GetTransientPackage());
-						}
-						else if (MapChangeType == EMapChangeType::LoadMap || MapChangeType == EMapChangeType::NewMap)
-						{
-							UWorld* World = GEditor->GetEditorWorldContext().World();
-							check(World);
-							CreatedWidget->Rename(*CreatedWidget->GetName(), World);
-						}
+						CreatedWidget->Rename(*CreatedWidget->GetName(), GetTransientPackage());
+					}
+					else if (MapChangeType == EMapChangeType::LoadMap || MapChangeType == EMapChangeType::NewMap)
+					{
+						UWorld* World = GEditor->GetEditorWorldContext().World();
+						check(World);
+						CreatedWidget->Rename(*CreatedWidget->GetName(), World);
 					}
 				}
 			}
@@ -281,21 +235,17 @@ public:
 
 	virtual void AddLoadedScriptUI(class UEditorUtilityWidgetBlueprint* InBlueprint) override
 	{
-		if (EditorUtilityContext)
-		{
-			EditorUtilityContext->LoadedUIs.AddUnique(InBlueprint);
-			EditorUtilityContext->SaveConfig();
-		}
+		UEditorUtilitySubsystem* EditorUtilitySubsystem = GEditor->GetEditorSubsystem<UEditorUtilitySubsystem>();
+		EditorUtilitySubsystem->LoadedUIs.AddUnique(InBlueprint);
+		EditorUtilitySubsystem->SaveConfig();
 	}
 
 
 	virtual void RemoveLoadedScriptUI(class UEditorUtilityWidgetBlueprint* InBlueprint) override
 	{
-		if (EditorUtilityContext)
-		{
-			EditorUtilityContext->LoadedUIs.Remove(InBlueprint);
-			EditorUtilityContext->SaveConfig();
-		}
+		UEditorUtilitySubsystem* EditorUtilitySubsystem = GEditor->GetEditorSubsystem<UEditorUtilitySubsystem>();
+		EditorUtilitySubsystem->LoadedUIs.Remove(InBlueprint);
+		EditorUtilitySubsystem->SaveConfig();	
 	}
 
 protected:
@@ -310,10 +260,6 @@ protected:
 
 	virtual void AddReferencedObjects(FReferenceCollector& Collector) override
 	{
-		if (EditorUtilityContext)
-		{
-			Collector.AddReferencedObject(EditorUtilityContext);
-		}
 	}
 
 	void OnPrepareToCleanseEditorObject(UObject* InObject)
@@ -353,28 +299,25 @@ protected:
 
 	void HandleAssetRemoved(const FAssetData& InAssetData)
 	{
-		if (EditorUtilityContext)
+		bool bDeletingLoadedUI = false;
+		for (const FSoftObjectPath& LoadedUIPath : GEditor->GetEditorSubsystem<UEditorUtilitySubsystem>()->LoadedUIs)
 		{
-			bool bDeletingLoadedUI = false;
-			for (FSoftObjectPath LoadedUIPath : EditorUtilityContext->LoadedUIs)
+			if (LoadedUIPath.GetAssetPathName() == InAssetData.ObjectPath)
 			{
-				if (LoadedUIPath.GetAssetPathName() == InAssetData.ObjectPath)
-				{
-					bDeletingLoadedUI = true;
-					break;
-				}
+				bDeletingLoadedUI = true;
+				break;
 			}
+		}
 
-			if (bDeletingLoadedUI)
+		if (bDeletingLoadedUI)
+		{
+			FName UIToCleanup = FName(*(InAssetData.ObjectPath.ToString() + LOCTEXT("ActiveTabSuffix", "_ActiveTab").ToString()));
+			FLevelEditorModule& LevelEditorModule = FModuleManager::GetModuleChecked<FLevelEditorModule>(TEXT("LevelEditor"));
+			TSharedPtr<FTabManager> LevelEditorTabManager = LevelEditorModule.GetLevelEditorTabManager();
+			TSharedPtr<SDockTab> CurrentTab = LevelEditorTabManager->FindExistingLiveTab(UIToCleanup);
+			if (CurrentTab.IsValid())
 			{
-				FName UIToCleanup = FName(*(InAssetData.ObjectPath.ToString() + LOCTEXT("ActiveTabSuffix", "_ActiveTab").ToString()));
-				FLevelEditorModule& LevelEditorModule = FModuleManager::GetModuleChecked<FLevelEditorModule>(TEXT("LevelEditor"));
-				TSharedPtr<FTabManager> LevelEditorTabManager = LevelEditorModule.GetLevelEditorTabManager();
-				TSharedPtr<SDockTab> CurrentTab = LevelEditorTabManager->FindExistingLiveTab(UIToCleanup);
-				if (CurrentTab.IsValid())
-				{
-					CurrentTab->RequestCloseTab();
-				}
+				CurrentTab->RequestCloseTab();
 			}
 		}
 	}
@@ -383,8 +326,6 @@ protected:
 	TSharedPtr<class FWorkspaceItem> ScriptedEditorWidgetsGroup;
 
 	EAssetTypeCategories::Type EditorUtilityAssetCategory;
-
-	UEditorUtilityContext* EditorUtilityContext;
 };
 
 
diff --git a/Engine/Source/Editor/Blutility/Private/EditorUtilityContext.cpp b/Engine/Source/Editor/Blutility/Private/EditorUtilityContext.cpp
deleted file mode 100644
diff --git a/Engine/Source/Editor/Blutility/Private/EditorUtilityContext.h b/Engine/Source/Editor/Blutility/Private/EditorUtilityContext.h
deleted file mode 100644
diff --git a/Engine/Source/Editor/Blutility/Private/EditorUtilityWidget.cpp b/Engine/Source/Editor/Blutility/Private/EditorUtilityWidget.cpp
--- a/Engine/Source/Editor/Blutility/Private/EditorUtilityWidget.cpp
+++ b/Engine/Source/Editor/Blutility/Private/EditorUtilityWidget.cpp
@@ -29,7 +29,7 @@ void UEditorUtilityWidget::ExecuteDefaultAction()
 	FScopedTransaction Transaction(NSLOCTEXT("UnrealEd", "BlutilityAction", "Blutility Action"));
 	FEditorScriptExecutionGuard ScriptGuard;
 
-	OnDefaultActionClicked();
+	Run();
 }
 
 
diff --git a/Engine/Source/Editor/UnrealEd/Private/EditorEngine.cpp b/Engine/Source/Editor/UnrealEd/Private/EditorEngine.cpp
--- a/Engine/Source/Editor/UnrealEd/Private/EditorEngine.cpp
+++ b/Engine/Source/Editor/UnrealEd/Private/EditorEngine.cpp
@@ -1152,6 +1152,8 @@ void UEditorEngine::FinishDestroy()
 			UE_LOG(LogEditor, Warning, TEXT("Warning: Play world is active"));
 		}
 
+		EditorSubsystemCollection.Deinitialize();
+
 		// Unregister events
 		FEditorDelegates::MapChange.RemoveAll(this);
 		FCoreDelegates::ModalErrorMessage.Unbind();
