Missing file
#rnx

diff --git a/Engine/Source/Editor/Blutility/Public/EditorUtilitySubsystem.h b/Engine/Source/Editor/Blutility/Public/EditorUtilitySubsystem.h
new file mode 100644
--- /dev/null
+++ b/Engine/Source/Editor/Blutility/Public/EditorUtilitySubsystem.h
@@ -0,0 +1,44 @@
+// Copyright 1998-2019 Epic Games, Inc. All Rights Reserved.
+
+#pragma once
+
+#include "EditorSubsystem.h"
+
+#include "UObject/ObjectMacros.h"
+#include "UObject/UObjectGlobals.h"
+
+#include "EditorUtilitySubsystem.generated.h"
+
+class SWindow;
+
+UCLASS(config = EditorPerProjectUserSettings)
+class BLUTILITY_API UEditorUtilitySubsystem : public UEditorSubsystem
+{
+	GENERATED_BODY()
+
+public:
+	UEditorUtilitySubsystem();
+
+	virtual void Initialize(FSubsystemCollectionBase& Collection);
+	virtual void Deinitialize();
+
+	void MainFrameCreationFinished(TSharedPtr<SWindow> InRootWindow, bool bIsNewProjectWindow);
+	void HandleStartup();
+
+	UPROPERTY(config)
+	TArray<FSoftObjectPath> LoadedUIs;
+
+	UPROPERTY(config)
+	TArray<FSoftObjectPath> StartupObjects;
+
+	// Allow startup object to be garbage collected
+	UFUNCTION(BlueprintCallable, Category = "Development|Editor")
+	void ReleaseInstanceOfAsset(UObject* Asset);
+
+	bool TryRun(UObject* Asset);
+
+private:
+	
+	UPROPERTY()
+	TMap<UObject*, UObject*> ObjectInstances;
+};
