Missing file
#rnx

diff --git a/Engine/Source/Editor/Blutility/Private/EditorUtilitySubsystem.cpp b/Engine/Source/Editor/Blutility/Private/EditorUtilitySubsystem.cpp
new file mode 100644
--- /dev/null
+++ b/Engine/Source/Editor/Blutility/Private/EditorUtilitySubsystem.cpp
@@ -0,0 +1,96 @@
+// Copyright 1998-2019 Epic Games, Inc. All Rights Reserved.
+
+#include "EditorUtilitySubsystem.h"
+#include "EditorUtilityCommon.h"
+#include "Interfaces/IMainFrameModule.h"
+
+UEditorUtilitySubsystem::UEditorUtilitySubsystem() :
+	UEditorSubsystem()
+{
+
+}
+
+void UEditorUtilitySubsystem::Initialize(FSubsystemCollectionBase& Collection)
+{
+	IMainFrameModule& MainFrameModule = IMainFrameModule::Get();
+	if (MainFrameModule.IsWindowInitialized())
+	{
+		HandleStartup();
+	}
+	else
+	{
+		MainFrameModule.OnMainFrameCreationFinished().AddUObject(this, &UEditorUtilitySubsystem::MainFrameCreationFinished);
+	}
+}
+
+void UEditorUtilitySubsystem::Deinitialize()
+{
+	if (FModuleManager::Get().IsModuleLoaded("MainFrame"))
+	{
+		IMainFrameModule::Get().OnMainFrameCreationFinished().RemoveAll(this);
+	}
+}
+
+void UEditorUtilitySubsystem::MainFrameCreationFinished(TSharedPtr<SWindow> InRootWindow, bool bIsNewProjectWindow)
+{
+	HandleStartup();
+}
+
+void UEditorUtilitySubsystem::HandleStartup()
+{
+	for (const FSoftObjectPath& ObjectPath : StartupObjects)
+	{
+		UObject* Object = ObjectPath.TryLoad();
+		if (!Object || Object->IsPendingKillOrUnreachable())
+		{
+			UE_LOG(LogEditorUtilityBlueprint, Warning, TEXT("Could not load: %s"), *ObjectPath.ToString());
+			continue;
+		}
+		
+		TryRun(ObjectPath.TryLoad());
+	}
+}
+
+bool UEditorUtilitySubsystem::TryRun(UObject* Asset)
+{
+	if (!Asset || Asset->IsPendingKillOrUnreachable())
+	{
+		UE_LOG(LogEditorUtilityBlueprint, Warning, TEXT("Could not run: %s"), *Asset->GetPathName());
+		return false;
+	}
+
+	UClass* ObjectClass = Asset->GetClass();
+	if (UBlueprint* Blueprint = Cast<UBlueprint>(Asset))
+	{
+		ObjectClass = Blueprint->GeneratedClass;
+	}
+
+	if (!ObjectClass)
+	{
+		UE_LOG(LogEditorUtilityBlueprint, Warning, TEXT("Missing class: %s"), *Asset->GetPathName());
+		return false;
+	}
+
+	static const FName RunFunctionName("Run");
+	UFunction* RunFunction = ObjectClass->FindFunctionByName(RunFunctionName);
+	if (RunFunction)
+	{
+		UObject* Instance = NewObject<UObject>(this, ObjectClass);
+		ObjectInstances.Add(Asset, Instance);
+
+		FEditorScriptExecutionGuard ScriptGuard;
+		Instance->ProcessEvent(RunFunction, nullptr);
+		return true;
+	}
+	else
+	{
+		UE_LOG(LogEditorUtilityBlueprint, Warning, TEXT("Missing function named 'Run': %s"), *Asset->GetPathName());
+	}
+
+	return false;
+}
+
+void UEditorUtilitySubsystem::ReleaseInstanceOfAsset(UObject* Asset)
+{
+	ObjectInstances.Remove(Asset);
+}
