Change for static analysis
#rnx

diff --git a/Engine/Source/Editor/Blutility/Private/EditorUtilitySubsystem.cpp b/Engine/Source/Editor/Blutility/Private/EditorUtilitySubsystem.cpp
--- a/Engine/Source/Editor/Blutility/Private/EditorUtilitySubsystem.cpp
+++ b/Engine/Source/Editor/Blutility/Private/EditorUtilitySubsystem.cpp
@@ -56,7 +56,7 @@ bool UEditorUtilitySubsystem::TryRun(UObject* Asset)
 {
 	if (!Asset || Asset->IsPendingKillOrUnreachable())
 	{
-		UE_LOG(LogEditorUtilityBlueprint, Warning, TEXT("Could not run: %s"), *Asset->GetPathName());
+		UE_LOG(LogEditorUtilityBlueprint, Warning, TEXT("Could not run: %s"), Asset ? *Asset->GetPathName() : TEXT("None"));
 		return false;
 	}
 
