Adding a spawn tab node for editor utility widgets to the editor utility subsystem
#rnx

diff --git a/Engine/Source/Editor/Blutility/Private/AssetTypeActions_EditorUtilityWidgetBlueprint.cpp b/Engine/Source/Editor/Blutility/Private/AssetTypeActions_EditorUtilityWidgetBlueprint.cpp
--- a/Engine/Source/Editor/Blutility/Private/AssetTypeActions_EditorUtilityWidgetBlueprint.cpp
+++ b/Engine/Source/Editor/Blutility/Private/AssetTypeActions_EditorUtilityWidgetBlueprint.cpp
@@ -17,6 +17,7 @@
 #include "Widgets/Docking/SDockTab.h"
 #include "Framework/Docking/TabManager.h"
 #include "IBlutilityModule.h"
+#include "EditorUtilitySubsystem.h"
 #include "SBlueprintDiff.h"
 
 #define LOCTEXT_NAMESPACE "AssetTypeActions"
@@ -50,7 +51,7 @@ void FAssetTypeActions_EditorUtilityWidgetBlueprint::GetActions(const TArray<UOb
 
 	MenuBuilder.AddMenuEntry(
 		LOCTEXT("EditorUtilityWidget_Edit", "Run Editor Utility Widget"),
-		LOCTEXT("EditorUtilityWidget_EditTooltip", "Runs the single action or opens the tab built by this Editor Utility Widget Blueprint."),
+		LOCTEXT("EditorUtilityWidget_EditTooltip", "Opens the tab built by this Editor Utility Widget Blueprint."),
 		FSlateIcon(),
 		FUIAction(
 			FExecuteAction::CreateSP(this, &FAssetTypeActions_EditorUtilityWidgetBlueprint::ExecuteRun, Blueprints),
@@ -116,31 +117,11 @@ void FAssetTypeActions_EditorUtilityWidgetBlueprint::ExecuteRun(FWeakBlueprintPo
 		{
 			if (Blueprint->GeneratedClass->IsChildOf(UEditorUtilityWidget::StaticClass()))
 			{
-				const UEditorUtilityWidget* CDO = Blueprint->GeneratedClass->GetDefaultObject<UEditorUtilityWidget>();
-				if (CDO->ShouldAutoRunDefaultAction())
+				UEditorUtilityWidgetBlueprint* EditorWidget = Cast<UEditorUtilityWidgetBlueprint>(Blueprint);
+				if (EditorWidget)
 				{
-					// This is an instant-run blueprint, just execute it
-					UEditorUtilityWidget* Instance = NewObject<UEditorUtilityWidget>(GetTransientPackage(), Blueprint->GeneratedClass);
-					Instance->ExecuteDefaultAction();
-				}
-				else
-				{
-					FName RegistrationName = FName(*(Blueprint->GetPathName() + LOCTEXT("ActiveTabSuffix", "_ActiveTab").ToString()));
-					FText DisplayName = FText::FromString(Blueprint->GetName());
-					FLevelEditorModule& LevelEditorModule = FModuleManager::GetModuleChecked<FLevelEditorModule>(TEXT("LevelEditor"));
-					TSharedPtr<FTabManager> LevelEditorTabManager = LevelEditorModule.GetLevelEditorTabManager();
-					if (!LevelEditorTabManager->HasTabSpawner(RegistrationName))
-					{
-						IBlutilityModule* BlutilityModule = FModuleManager::GetModulePtr<IBlutilityModule>("Blutility");
-						UEditorUtilityWidgetBlueprint* WidgetBlueprint = Cast<UEditorUtilityWidgetBlueprint>(Blueprint);
-						WidgetBlueprint->SetRegistrationName(RegistrationName);
-						LevelEditorTabManager->RegisterTabSpawner(RegistrationName, FOnSpawnTab::CreateUObject(WidgetBlueprint, &UEditorUtilityWidgetBlueprint::SpawnEditorUITab))
-							.SetDisplayName(DisplayName)
-							.SetGroup(BlutilityModule->GetMenuGroup().ToSharedRef());
-						BlutilityModule->AddLoadedScriptUI(WidgetBlueprint);
-					}
-					TSharedRef<SDockTab> NewDockTab = LevelEditorTabManager->InvokeTab(RegistrationName);
-
+					UEditorUtilitySubsystem* EditorUtilitySubsystem = GEditor->GetEditorSubsystem<UEditorUtilitySubsystem>();
+					EditorUtilitySubsystem->SpawnAndRegisterTab(EditorWidget);
 				}
 			}
 		}
diff --git a/Engine/Source/Editor/Blutility/Private/EditorUtilitySubsystem.cpp b/Engine/Source/Editor/Blutility/Private/EditorUtilitySubsystem.cpp
--- a/Engine/Source/Editor/Blutility/Private/EditorUtilitySubsystem.cpp
+++ b/Engine/Source/Editor/Blutility/Private/EditorUtilitySubsystem.cpp
@@ -4,8 +4,15 @@
 #include "EditorUtilityCommon.h"
 #include "Interfaces/IMainFrameModule.h"
 #include "Engine/Blueprint.h"
+#include "EditorUtilityWidgetBlueprint.h"
+#include "LevelEditor.h"
+#include "IBlutilityModule.h"
+#include "EditorUtilityWidget.h"
 #include "ScopedTransaction.h"
 
+#define LOCTEXT_NAMESPACE "EditorUtilitySubsystem"
+
+
 UEditorUtilitySubsystem::UEditorUtilitySubsystem() :
 	UEditorSubsystem()
 {
@@ -96,3 +103,30 @@ void UEditorUtilitySubsystem::ReleaseInstanceOfAsset(UObject* Asset)
 {
 	ObjectInstances.Remove(Asset);
 }
+
+UEditorUtilityWidget* UEditorUtilitySubsystem::SpawnAndRegisterTab(UEditorUtilityWidgetBlueprint* InBlueprint)
+{
+	if (IsRunningCommandlet())
+	{
+		return nullptr;
+	}
+
+	FName RegistrationName = FName(*(InBlueprint->GetPathName() + LOCTEXT("ActiveTabSuffix", "_ActiveTab").ToString()));
+	FText DisplayName = FText::FromString(InBlueprint->GetName());
+	FLevelEditorModule& LevelEditorModule = FModuleManager::GetModuleChecked<FLevelEditorModule>(TEXT("LevelEditor"));
+	TSharedPtr<FTabManager> LevelEditorTabManager = LevelEditorModule.GetLevelEditorTabManager();
+	if (!LevelEditorTabManager->HasTabSpawner(RegistrationName))
+	{
+		IBlutilityModule* BlutilityModule = FModuleManager::GetModulePtr<IBlutilityModule>("Blutility");
+		LevelEditorTabManager->RegisterTabSpawner(RegistrationName, FOnSpawnTab::CreateUObject(InBlueprint, &UEditorUtilityWidgetBlueprint::SpawnEditorUITab))
+			.SetDisplayName(DisplayName)
+			.SetGroup(BlutilityModule->GetMenuGroup().ToSharedRef());
+		InBlueprint->SetRegistrationName(RegistrationName);
+		BlutilityModule->AddLoadedScriptUI(InBlueprint);
+	}
+	TSharedRef<SDockTab> NewDockTab = LevelEditorTabManager->InvokeTab(RegistrationName);
+	return InBlueprint->GetCreatedWidget();
+}
+
+
+#undef LOCTEXT_NAMESPACE
diff --git a/Engine/Source/Editor/Blutility/Public/EditorUtilitySubsystem.h b/Engine/Source/Editor/Blutility/Public/EditorUtilitySubsystem.h
--- a/Engine/Source/Editor/Blutility/Public/EditorUtilitySubsystem.h
+++ b/Engine/Source/Editor/Blutility/Public/EditorUtilitySubsystem.h
@@ -10,6 +10,7 @@
 #include "EditorUtilitySubsystem.generated.h"
 
 class SWindow;
+class UEditorUtilityWidget;
 
 UCLASS(config = EditorPerProjectUserSettings)
 class BLUTILITY_API UEditorUtilitySubsystem : public UEditorSubsystem
@@ -37,6 +38,9 @@ public:
 
 	bool TryRun(UObject* Asset);
 
+	UFUNCTION(BlueprintCallable, Category = "Development|Editor")
+	UEditorUtilityWidget* SpawnAndRegisterTab(class UEditorUtilityWidgetBlueprint* InBlueprint);
+
 private:
 	
 	UPROPERTY()
