Fixing name of Editor Utility Widget in Content Browser
#rnx

diff --git a/Engine/Source/Editor/Blutility/Private/AssetTypeActions_EditorUtilityWidgetBlueprint.cpp b/Engine/Source/Editor/Blutility/Private/AssetTypeActions_EditorUtilityWidgetBlueprint.cpp
--- a/Engine/Source/Editor/Blutility/Private/AssetTypeActions_EditorUtilityWidgetBlueprint.cpp
+++ b/Engine/Source/Editor/Blutility/Private/AssetTypeActions_EditorUtilityWidgetBlueprint.cpp
@@ -27,7 +27,7 @@
 
 FText FAssetTypeActions_EditorUtilityWidgetBlueprint::GetName() const
 {
-	return LOCTEXT("AssetTypeActions_EditorUtilityWidget", "Editor Widget");
+	return LOCTEXT("AssetTypeActions_EditorUtilityWidget", "Editor Utility Widget");
 }
 
 FColor FAssetTypeActions_EditorUtilityWidgetBlueprint::GetTypeColor() const
