Adding conditional to check for a Blueprint being passed in
#rnx

diff --git a/Engine/Source/Editor/Blutility/Private/EditorUtilitySubsystem.cpp b/Engine/Source/Editor/Blutility/Private/EditorUtilitySubsystem.cpp
--- a/Engine/Source/Editor/Blutility/Private/EditorUtilitySubsystem.cpp
+++ b/Engine/Source/Editor/Blutility/Private/EditorUtilitySubsystem.cpp
@@ -106,26 +106,27 @@ void UEditorUtilitySubsystem::ReleaseInstanceOfAsset(UObject* Asset)
 
 UEditorUtilityWidget* UEditorUtilitySubsystem::SpawnAndRegisterTab(UEditorUtilityWidgetBlueprint* InBlueprint)
 {
-	if (IsRunningCommandlet())
-	{
-		return nullptr;
-	}
 
-	FName RegistrationName = FName(*(InBlueprint->GetPathName() + LOCTEXT("ActiveTabSuffix", "_ActiveTab").ToString()));
-	FText DisplayName = FText::FromString(InBlueprint->GetName());
-	FLevelEditorModule& LevelEditorModule = FModuleManager::GetModuleChecked<FLevelEditorModule>(TEXT("LevelEditor"));
-	TSharedPtr<FTabManager> LevelEditorTabManager = LevelEditorModule.GetLevelEditorTabManager();
-	if (!LevelEditorTabManager->HasTabSpawner(RegistrationName))
+	if (InBlueprint && !IsRunningCommandlet())
 	{
-		IBlutilityModule* BlutilityModule = FModuleManager::GetModulePtr<IBlutilityModule>("Blutility");
-		LevelEditorTabManager->RegisterTabSpawner(RegistrationName, FOnSpawnTab::CreateUObject(InBlueprint, &UEditorUtilityWidgetBlueprint::SpawnEditorUITab))
-			.SetDisplayName(DisplayName)
-			.SetGroup(BlutilityModule->GetMenuGroup().ToSharedRef());
-		InBlueprint->SetRegistrationName(RegistrationName);
-		BlutilityModule->AddLoadedScriptUI(InBlueprint);
+		FName RegistrationName = FName(*(InBlueprint->GetPathName() + LOCTEXT("ActiveTabSuffix", "_ActiveTab").ToString()));
+		FText DisplayName = FText::FromString(InBlueprint->GetName());
+		FLevelEditorModule& LevelEditorModule = FModuleManager::GetModuleChecked<FLevelEditorModule>(TEXT("LevelEditor"));
+		TSharedPtr<FTabManager> LevelEditorTabManager = LevelEditorModule.GetLevelEditorTabManager();
+		if (!LevelEditorTabManager->HasTabSpawner(RegistrationName))
+		{
+			IBlutilityModule* BlutilityModule = FModuleManager::GetModulePtr<IBlutilityModule>("Blutility");
+			LevelEditorTabManager->RegisterTabSpawner(RegistrationName, FOnSpawnTab::CreateUObject(InBlueprint, &UEditorUtilityWidgetBlueprint::SpawnEditorUITab))
+				.SetDisplayName(DisplayName)
+				.SetGroup(BlutilityModule->GetMenuGroup().ToSharedRef());
+			InBlueprint->SetRegistrationName(RegistrationName);
+			BlutilityModule->AddLoadedScriptUI(InBlueprint);
+		}
+		TSharedRef<SDockTab> NewDockTab = LevelEditorTabManager->InvokeTab(RegistrationName);
+		return InBlueprint->GetCreatedWidget();
 	}
-	TSharedRef<SDockTab> NewDockTab = LevelEditorTabManager->InvokeTab(RegistrationName);
-	return InBlueprint->GetCreatedWidget();
+
+	return nullptr;
 }
 
 
