When changing levels, recreate the tab widget to avoid holding a reference to the previous level

diff --git a/Engine/Source/Editor/Blutility/Classes/EditorUtilityWidgetBlueprint.h b/Engine/Source/Editor/Blutility/Classes/EditorUtilityWidgetBlueprint.h
--- a/Engine/Source/Editor/Blutility/Classes/EditorUtilityWidgetBlueprint.h
+++ b/Engine/Source/Editor/Blutility/Classes/EditorUtilityWidgetBlueprint.h
@@ -16,6 +16,7 @@
 class UBlueprint;
 class UEditorUtilityWidget;
 enum class EAssetEditorCloseReason : uint8;
+enum class EMapChangeType : uint8;
 
 UCLASS()
 class UEditorUtilityWidgetBlueprint : public UWidgetBlueprint
@@ -59,7 +60,8 @@ public:
 protected:
 	virtual void LoadModulesRequiredForCompilation() override;
 #endif
-	
+	void ChangeTabWorld(UWorld* World, EMapChangeType MapChangeType);
+
 private:
 	FName RegistrationName;
 
diff --git a/Engine/Source/Editor/Blutility/Private/EditorUtilityWidgetBlueprint.cpp b/Engine/Source/Editor/Blutility/Private/EditorUtilityWidgetBlueprint.cpp
--- a/Engine/Source/Editor/Blutility/Private/EditorUtilityWidgetBlueprint.cpp
+++ b/Engine/Source/Editor/Blutility/Private/EditorUtilityWidgetBlueprint.cpp
@@ -57,6 +57,9 @@ TSharedRef<SDockTab> UEditorUtilityWidgetBlueprint::SpawnEditorUITab(const FSpaw
 	CreatedTab = SpawnedTab;
 	
 	GEditor->OnBlueprintReinstanced().AddUObject(this, &UEditorUtilityWidgetBlueprint::RegenerateCreatedTab);
+	
+	FLevelEditorModule& LevelEditor = FModuleManager::LoadModuleChecked<FLevelEditorModule>("LevelEditor");
+	LevelEditor.OnMapChanged().AddUObject(this, &UEditorUtilityWidgetBlueprint::ChangeTabWorld);
 
 	return SpawnedTab;
 }
@@ -94,6 +97,23 @@ void UEditorUtilityWidgetBlueprint::RegenerateCreatedTab()
 	}
 }
 
+void UEditorUtilityWidgetBlueprint::ChangeTabWorld(UWorld* World, EMapChangeType MapChangeType)
+{
+	if (MapChangeType == EMapChangeType::TearDownWorld)
+	{
+		CreatedUMGWidget = nullptr;
+		if (CreatedTab.IsValid())
+		{
+			CreatedTab.Pin()->SetContent(SNullWidget::NullWidget);
+		}
+	}
+	else if (MapChangeType != EMapChangeType::SaveMap)
+	{
+		// Recreate the widget if we are loading a map or opening a new map
+		RegenerateCreatedTab();
+	}
+}
+
 void UEditorUtilityWidgetBlueprint::UpdateRespawnListIfNeeded(TSharedRef<SDockTab> TabBeingClosed)
 {
 	const UEditorUtilityWidget* EditorUtilityWidget = GeneratedClass->GetDefaultObject<UEditorUtilityWidget>();
