Made the created widget transient to allow saving the editor utility widget while it is open in a tab

diff --git a/Engine/Source/Editor/Blutility/Classes/EditorUtilityWidgetBlueprint.h b/Engine/Source/Editor/Blutility/Classes/EditorUtilityWidgetBlueprint.h
--- a/Engine/Source/Editor/Blutility/Classes/EditorUtilityWidgetBlueprint.h
+++ b/Engine/Source/Editor/Blutility/Classes/EditorUtilityWidgetBlueprint.h
@@ -67,7 +67,7 @@ private:
 
 	TWeakPtr<SDockTab> CreatedTab;
 
-	UPROPERTY()
+	UPROPERTY(Transient)
 	UEditorUtilityWidget* CreatedUMGWidget;
 
 };
