Removing LoadModulesRequiredForCompilation as it was causing editor utilities to be loaded in game and is not needed.
#rnx

diff --git a/Engine/Source/Editor/Blutility/Classes/EditorUtilityBlueprint.h b/Engine/Source/Editor/Blutility/Classes/EditorUtilityBlueprint.h
--- a/Engine/Source/Editor/Blutility/Classes/EditorUtilityBlueprint.h
+++ b/Engine/Source/Editor/Blutility/Classes/EditorUtilityBlueprint.h
@@ -21,8 +21,5 @@ class BLUTILITY_API UEditorUtilityBlueprint : public UBlueprint
 	{
 		return false;
 	}
-
-protected:
-	virtual void LoadModulesRequiredForCompilation() override;
 // End of UBlueprint interface
 };
diff --git a/Engine/Source/Editor/Blutility/Classes/EditorUtilityWidgetBlueprint.h b/Engine/Source/Editor/Blutility/Classes/EditorUtilityWidgetBlueprint.h
--- a/Engine/Source/Editor/Blutility/Classes/EditorUtilityWidgetBlueprint.h
+++ b/Engine/Source/Editor/Blutility/Classes/EditorUtilityWidgetBlueprint.h
@@ -56,10 +56,6 @@ public:
 		return RegistrationName;
 	}
 
-#if WITH_EDITORONLY_DATA
-protected:
-	virtual void LoadModulesRequiredForCompilation() override;
-#endif
 	void ChangeTabWorld(UWorld* World, EMapChangeType MapChangeType);
 
 private:
diff --git a/Engine/Source/Editor/Blutility/Private/EditorUtilityBlueprint.cpp b/Engine/Source/Editor/Blutility/Private/EditorUtilityBlueprint.cpp
--- a/Engine/Source/Editor/Blutility/Private/EditorUtilityBlueprint.cpp
+++ b/Engine/Source/Editor/Blutility/Private/EditorUtilityBlueprint.cpp
@@ -10,13 +10,3 @@ UEditorUtilityBlueprint::UEditorUtilityBlueprint(const FObjectInitializer& Objec
 	: Super(ObjectInitializer)
 {
 }
-
-#if WITH_EDITORONLY_DATA
-void UEditorUtilityBlueprint::LoadModulesRequiredForCompilation()
-{
-	Super::LoadModulesRequiredForCompilation();
-
-	static const FName ModuleName(TEXT("Blutility"));
-	FModuleManager::Get().LoadModule(ModuleName);
-}
-#endif //WITH_EDITORONLY_DATA
\ No newline at end of file
diff --git a/Engine/Source/Editor/Blutility/Private/EditorUtilityWidgetBlueprint.cpp b/Engine/Source/Editor/Blutility/Private/EditorUtilityWidgetBlueprint.cpp
--- a/Engine/Source/Editor/Blutility/Private/EditorUtilityWidgetBlueprint.cpp
+++ b/Engine/Source/Editor/Blutility/Private/EditorUtilityWidgetBlueprint.cpp
@@ -131,13 +131,4 @@ void UEditorUtilityWidgetBlueprint::GetReparentingRules(TSet< const UClass* >& A
 	AllowedChildrenOfClasses.Add(UEditorUtilityWidget::StaticClass());
 }
 
-#if WITH_EDITORONLY_DATA
-void UEditorUtilityWidgetBlueprint::LoadModulesRequiredForCompilation()
-{
-	Super::LoadModulesRequiredForCompilation();
-
-	static const FName ModuleName(TEXT("Blutility"));
-	FModuleManager::Get().LoadModule(ModuleName);
-}
-#endif //WITH_EDITORONLY_DATA
 
