Fixing variable shadowing

diff --git a/Engine/Source/Editor/Blutility/Private/BlutilityModule.cpp b/Engine/Source/Editor/Blutility/Private/BlutilityModule.cpp
--- a/Engine/Source/Editor/Blutility/Private/BlutilityModule.cpp
+++ b/Engine/Source/Editor/Blutility/Private/BlutilityModule.cpp
@@ -120,7 +120,7 @@ public:
 		}
 	}
 
-	void OnMapChanged(UWorld* World, EMapChangeType MapChangeType)
+	void OnMapChanged(UWorld* InWorld, EMapChangeType MapChangeType)
 	{
 		if (EditorUtilityContext)
 		{
