Removing forward enum declaration

diff --git a/Engine/Source/Editor/Blutility/Classes/EditorUtilityWidgetBlueprintFactory.h b/Engine/Source/Editor/Blutility/Classes/EditorUtilityWidgetBlueprintFactory.h
--- a/Engine/Source/Editor/Blutility/Classes/EditorUtilityWidgetBlueprintFactory.h
+++ b/Engine/Source/Editor/Blutility/Classes/EditorUtilityWidgetBlueprintFactory.h
@@ -10,8 +10,10 @@
 #include "UObject/ObjectMacros.h"
 #include "Templates/SubclassOf.h"
 #include "Factories/Factory.h"
+#include "Engine/Blueprint.h"
 #include "EditorUtilityWidgetBlueprintFactory.generated.h"
 
+
 class SWindow;
 
 UCLASS(HideCategories = Object, MinimalAPI)
@@ -21,7 +23,7 @@ class UEditorUtilityWidgetBlueprintFactory : public UFactory
 
 	// The type of blueprint that will be created
 	UPROPERTY(EditAnywhere, Category = WidgetBlueprintFactory)
-	TEnumAsByte<enum EBlueprintType> BlueprintType;
+	TEnumAsByte<EBlueprintType> BlueprintType;
 
 	// The parent class of the created blueprint
 	UPROPERTY(EditAnywhere, Category = WidgetBlueprintFactory, meta = (AllowAbstract = ""))
diff --git a/Engine/Source/Editor/Blutility/Private/EditorUtilityWidgetBlueprintFactory.cpp b/Engine/Source/Editor/Blutility/Private/EditorUtilityWidgetBlueprintFactory.cpp
--- a/Engine/Source/Editor/Blutility/Private/EditorUtilityWidgetBlueprintFactory.cpp
+++ b/Engine/Source/Editor/Blutility/Private/EditorUtilityWidgetBlueprintFactory.cpp
@@ -4,7 +4,6 @@
 #include "Misc/MessageDialog.h"
 #include "Modules/ModuleManager.h"
 #include "Widgets/SWindow.h"
-#include "Engine/Blueprint.h"
 #include "Settings/EditorExperimentalSettings.h"
 #include "Engine/BlueprintGeneratedClass.h"
 #include "EditorUtilityBlueprint.h"
