Only show a single execution button for call-in-editor Blutility functions

They should use the one implemented by FEditorUtilityInstanceDetails, as that calls the post-action logic correctly

diff --git a/Engine/Source/Editor/Blutility/Private/BlutilityDetailsPanel.cpp b/Engine/Source/Editor/Blutility/Private/BlutilityDetailsPanel.cpp
--- a/Engine/Source/Editor/Blutility/Private/BlutilityDetailsPanel.cpp
+++ b/Engine/Source/Editor/Blutility/Private/BlutilityDetailsPanel.cpp
@@ -5,6 +5,7 @@
 #include "Widgets/Layout/SWrapBox.h"
 #include "Widgets/Text/STextBlock.h"
 #include "Widgets/Input/SButton.h"
+#include "EdGraphSchema_K2.h"
 #include "EditorUtilityBlueprint.h"
 #include "UObject/UnrealType.h"
 #include "GlobalEditorUtilityBase.h"
@@ -75,7 +76,7 @@ void FEditorUtilityInstanceDetails::CustomizeDetails(IDetailLayoutBuilder& Detai
 		IDetailCategoryBuilder& ActionsCategory = DetailLayoutBuilder.EditCategory(*CategoryName);
 
 		const APlacedEditorUtilityBase* PlacedActorCDO = Cast<const APlacedEditorUtilityBase>(Class->GetDefaultObject());
-		if (PlacedActorCDO)
+		if (PlacedActorCDO && !PlacedActorCDO->HelpText.IsEmpty())
 		{
 			ActionsCategory.AddCustomRow( FText::FromString(PlacedActorCDO->HelpText) )
 			[
@@ -85,7 +86,7 @@ void FEditorUtilityInstanceDetails::CustomizeDetails(IDetailLayoutBuilder& Detai
 		}
 		
 		const UGlobalEditorUtilityBase* GlobalBlutilityCDO = Cast<const UGlobalEditorUtilityBase>(Class->GetDefaultObject());
-		if (GlobalBlutilityCDO)
+		if (GlobalBlutilityCDO && !GlobalBlutilityCDO->HelpText.IsEmpty())
 		{
 			ActionsCategory.AddCustomRow( FText::FromString(GlobalBlutilityCDO->HelpText) )
 			[
@@ -101,7 +102,7 @@ void FEditorUtilityInstanceDetails::CustomizeDetails(IDetailLayoutBuilder& Detai
 		{
 			UFunction* Function = *FuncIt;
 
-			const bool bCallInEditorFunc = Function->GetBoolMetaData( TEXT("CallInEditor") );
+			const bool bCallInEditorFunc = Function->GetBoolMetaData(FBlueprintMetadata::MD_CallInEditor);
 			const bool bCanExecute = (Function->NumParms == 0) && bCallInEditorFunc;
 
 			if (bCanExecute)
diff --git a/Engine/Source/Editor/DetailCustomizations/Private/ObjectDetails.cpp b/Engine/Source/Editor/DetailCustomizations/Private/ObjectDetails.cpp
--- a/Engine/Source/Editor/DetailCustomizations/Private/ObjectDetails.cpp
+++ b/Engine/Source/Editor/DetailCustomizations/Private/ObjectDetails.cpp
@@ -15,6 +15,7 @@
 #include "DetailCategoryBuilder.h"
 #include "IDetailsView.h"
 #include "EdGraphSchema_K2.h"
+#include "Kismet2/BlueprintEditorUtils.h"
 #include "Widgets/SToolTip.h"
 #include "IDocumentation.h"
 #include "ObjectEditorUtils.h"
@@ -97,6 +98,18 @@ void FObjectDetails::AddCallInEditorMethods(IDetailLayoutBuilder& DetailBuilder)
 
 		if (TestFunction->GetBoolMetaData(FBlueprintMetadata::MD_CallInEditor) && (TestFunction->ParmsSize == 0))
 		{
+			if (UClass* TestFunctionOwnerClass = TestFunction->GetOwnerClass())
+			{
+				if (UBlueprint* Blueprint = Cast<UBlueprint>(TestFunctionOwnerClass->ClassGeneratedBy))
+				{
+					if (FBlueprintEditorUtils::IsBlutility(Blueprint))
+					{
+						// Skip Blutilities as these are handled by FEditorUtilityInstanceDetails
+						continue;
+					}
+				}
+			}
+
 			CallInEditorFunctions.Add(*FunctionIter);
 		}
 	}
