Fix for other automation tests failing w/editor only functions

diff --git a/Engine/Source/Editor/Blutility/Classes/EditorAutomationActorComponent.h b/Engine/Source/Editor/Blutility/Classes/EditorAutomationActorComponent.h
new file mode 100644
--- /dev/null
+++ b/Engine/Source/Editor/Blutility/Classes/EditorAutomationActorComponent.h
@@ -0,0 +1,21 @@
+// Copyright 1998-2018 Epic Games, Inc. All Rights Reserved.
+
+/**
+ * Base class of any editor-only automation test actor components
+ */
+
+#pragma once
+
+#include "CoreMinimal.h"
+#include "Components/ActorComponent.h"
+#include "UObject/ScriptMacros.h"
+#include "EditorAutomationActorComponent.generated.h"
+
+
+UCLASS(Abstract, Blueprintable)
+class BLUTILITY_API UEditorAutomationActorComponent : public UActorComponent
+{
+	GENERATED_UCLASS_BODY()
+
+	
+};
diff --git a/Engine/Source/Editor/Blutility/Private/EditorAutomationActorComponent.cpp b/Engine/Source/Editor/Blutility/Private/EditorAutomationActorComponent.cpp
new file mode 100644
--- /dev/null
+++ b/Engine/Source/Editor/Blutility/Private/EditorAutomationActorComponent.cpp
@@ -0,0 +1,11 @@
+// Copyright 1998-2018 Epic Games, Inc. All Rights Reserved.
+
+#include "EditorAutomationActorComponent.h"
+
+
+/////////////////////////////////////////////////////
+
+UEditorAutomationActorComponent::UEditorAutomationActorComponent(const FObjectInitializer& ObjectInitializer)
+	: Super(ObjectInitializer)
+{
+}
