(4.21.2) Prevent Blutility object from being destroyed by GC while script are executing
#rnx

diff --git a/Engine/Source/Editor/Blutility/Private/BlutilityMenuExtensions.cpp b/Engine/Source/Editor/Blutility/Private/BlutilityMenuExtensions.cpp
--- a/Engine/Source/Editor/Blutility/Private/BlutilityMenuExtensions.cpp
+++ b/Engine/Source/Editor/Blutility/Private/BlutilityMenuExtensions.cpp
@@ -279,6 +279,7 @@ void FBlutilityMenuExtensions::CreateBlutilityActionsMenu(FMenuBuilder& MenuBuil
 							{
 								// We dont run this on the CDO, as bad things could occur!
 								UObject* TempObject = NewObject<UObject>(GetTransientPackage(), FunctionAndUtil.Util->GetClass());
+								TempObject->AddToRoot(); // Some Blutility actions might run GC so the TempObject needs to be rooted to avoid getting destroyed
 
 								if(FunctionAndUtil.Function->NumParms > 0)
 								{
@@ -321,6 +322,8 @@ void FBlutilityMenuExtensions::CreateBlutilityActionsMenu(FMenuBuilder& MenuBuil
 									FEditorScriptExecutionGuard ScriptGuard;
 									TempObject->ProcessEvent(FunctionAndUtil.Function, nullptr);
 								}
+
+								TempObject->RemoveFromRoot();
 							}
 						}));
 				}
