Editor Utilities now respond to editor requests to cleanse object instances by clearing their references to any objects pending purge

This fixes crashes where they would take references to things like an actor within a world, and then the active editor world was changed (which verifies against GC reference leaks).

diff --git a/Engine/Plugins/Experimental/PythonScriptPlugin/Source/PythonScriptPlugin/Private/PyReferenceCollector.cpp b/Engine/Plugins/Experimental/PythonScriptPlugin/Source/PythonScriptPlugin/Private/PyReferenceCollector.cpp
--- a/Engine/Plugins/Experimental/PythonScriptPlugin/Source/PythonScriptPlugin/Private/PyReferenceCollector.cpp
+++ b/Engine/Plugins/Experimental/PythonScriptPlugin/Source/PythonScriptPlugin/Private/PyReferenceCollector.cpp
@@ -5,59 +5,10 @@
 #include "PyWrapperDelegate.h"
 #include "UObject/UnrealType.h"
 #include "UObject/UObjectHash.h"
+#include "UObject/PurgingReferenceCollector.h"
 
 #if WITH_PYTHON
 
-/** Reference collector that will purge (null) any references to the given set of objects (as if they'd bee marked PendingKill) */
-class FPyPurgingReferenceCollector : public FReferenceCollector
-{
-public:
-	bool HasObjectToPurge() const
-	{
-		return ObjectsToPurge.Num() > 0;
-	}
-
-	void AddObjectToPurge(const UObject* Object)
-	{
-		ObjectsToPurge.Add(Object);
-	}
-
-	virtual bool IsIgnoringArchetypeRef() const override
-	{
-		return false;
-	}
-
-	virtual bool IsIgnoringTransient() const override
-	{
-		return false;
-	}
-
-protected:
-	virtual void HandleObjectReference(UObject*& Object, const UObject* ReferencingObject, const UProperty* ReferencingProperty) override
-	{
-		ConditionalPurgeObject(Object);
-	}
-
-	virtual void HandleObjectReferences(UObject** InObjects, const int32 ObjectNum, const UObject* InReferencingObject, const UProperty* InReferencingProperty) override
-	{
-		for (int32 ObjectIndex = 0; ObjectIndex < ObjectNum; ++ObjectIndex)
-		{
-			UObject*& Object = InObjects[ObjectIndex];
-			ConditionalPurgeObject(Object);
-		}
-	}
-
-	void ConditionalPurgeObject(UObject*& Object)
-	{
-		if (ObjectsToPurge.Contains(Object))
-		{
-			Object = nullptr;
-		}
-	}
-
-	TSet<const UObject*> ObjectsToPurge;
-};
-
 FPyReferenceCollector& FPyReferenceCollector::Get()
 {
 	static FPyReferenceCollector Instance;
@@ -93,7 +44,7 @@ void FPyReferenceCollector::PurgeUnrealObjectReferences(const UObject* InObject,
 
 void FPyReferenceCollector::PurgeUnrealObjectReferences(const TArrayView<const UObject*>& InObjects, const bool bIncludeInnerObjects)
 {
-	FPyPurgingReferenceCollector PurgingReferenceCollector;
+	FPurgingReferenceCollector PurgingReferenceCollector;
 
 	for (const UObject* Object : InObjects)
 	{
diff --git a/Engine/Source/Editor/Blutility/Private/BlutilityModule.cpp b/Engine/Source/Editor/Blutility/Private/BlutilityModule.cpp
--- a/Engine/Source/Editor/Blutility/Private/BlutilityModule.cpp
+++ b/Engine/Source/Editor/Blutility/Private/BlutilityModule.cpp
@@ -5,8 +5,9 @@
 #include "Widgets/DeclarativeSyntaxSupport.h"
 #include "Framework/Docking/TabManager.h"
 #include "IBlutilityModule.h"
+#include "EditorUtilityWidget.h"
 #include "EditorUtilityBlueprint.h"
-
+#include "GlobalEditorUtilityBase.h"
 
 #include "AssetToolsModule.h"
 #include "PropertyEditorModule.h"
@@ -30,7 +31,8 @@
 #include "LevelEditor.h"
 #include "Editor.h"
 #include "UnrealEdMisc.h"
-
+#include "EditorSupportDelegates.h"
+#include "UObject/PurgingReferenceCollector.h"
 
 #define LOCTEXT_NAMESPACE "AssetTypeActions"
 
@@ -91,6 +93,8 @@ public:
 		FLevelEditorModule& LevelEditorModule = FModuleManager::LoadModuleChecked<FLevelEditorModule>(TEXT("LevelEditor"));
 		LevelEditorModule.OnTabManagerChanged().AddRaw(this, &FBlutilityModule::ReinitializeUIs);
 		LevelEditorModule.OnMapChanged().AddRaw(this, &FBlutilityModule::OnMapChanged);
+
+		FEditorSupportDelegates::PrepareToCleanseEditorObject.AddRaw(this, &FBlutilityModule::OnPrepareToCleanseEditorObject);
 	}
 
 	void ReinitializeUIs()
@@ -184,6 +188,8 @@ public:
 			PropertyModule.UnregisterCustomClassLayout("GlobalEditorUtilityBase");
 			PropertyModule.NotifyCustomizationModuleChanged();
 		}
+
+		FEditorSupportDelegates::PrepareToCleanseEditorObject.RemoveAll(this);
 	}
 
 	virtual bool IsBlutility( const UBlueprint* Blueprint ) const override
@@ -244,6 +250,41 @@ protected:
 		}
 	}
 
+	void OnPrepareToCleanseEditorObject(UObject* InObject)
+	{
+		// Gather the list of live Editor Utility instances to purge references from
+		TArray<UObject*> EditorUtilityInstances;
+		ForEachObjectOfClass(UEditorUtilityWidget::StaticClass(), [&EditorUtilityInstances](UObject* InEditorUtilityInstance)
+		{
+			EditorUtilityInstances.Add(InEditorUtilityInstance);
+		});
+		ForEachObjectOfClass(UGlobalEditorUtilityBase::StaticClass(), [&EditorUtilityInstances](UObject* InEditorUtilityInstance)
+		{
+			EditorUtilityInstances.Add(InEditorUtilityInstance);
+		});
+
+		if (EditorUtilityInstances.Num() > 0)
+		{
+			// Build up the complete list of objects to purge
+			FPurgingReferenceCollector PurgingReferenceCollector;
+			PurgingReferenceCollector.AddObjectToPurge(InObject);
+			ForEachObjectWithOuter(InObject, [&PurgingReferenceCollector](UObject* InInnerObject)
+			{
+				PurgingReferenceCollector.AddObjectToPurge(InInnerObject);
+			}, true);
+
+			// Run the purge for each Editor Utility instance
+			FReferenceCollectorArchive& PurgingReferenceCollectorArchive = PurgingReferenceCollector.GetVerySlowReferenceCollectorArchive();
+			for (UObject* EditorUtilityInstance : EditorUtilityInstances)
+			{
+				PurgingReferenceCollectorArchive.SetSerializingObject(EditorUtilityInstance);
+				EditorUtilityInstance->Serialize(PurgingReferenceCollectorArchive);
+				EditorUtilityInstance->CallAddReferencedObjects(PurgingReferenceCollector);
+				PurgingReferenceCollectorArchive.SetSerializingObject(nullptr);
+			}
+		}
+	}
+
 	/** Scripted Editor Widgets workspace menu item */
 	TSharedPtr<class FWorkspaceItem> ScriptedEditorWidgetsGroup;
 
diff --git a/Engine/Source/Runtime/CoreUObject/Public/UObject/PurgingReferenceCollector.h b/Engine/Source/Runtime/CoreUObject/Public/UObject/PurgingReferenceCollector.h
new file mode 100644
--- /dev/null
+++ b/Engine/Source/Runtime/CoreUObject/Public/UObject/PurgingReferenceCollector.h
@@ -0,0 +1,55 @@
+// Copyright 1998-2019 Epic Games, Inc. All Rights Reserved.
+
+#pragma once
+
+#include "UObject/UnrealType.h"
+
+/** Reference collector that will purge (null) any references to the given set of objects (as if they'd been marked PendingKill) */
+class FPurgingReferenceCollector : public FReferenceCollector
+{
+public:
+	bool HasObjectToPurge() const
+	{
+		return ObjectsToPurge.Num() > 0;
+	}
+
+	void AddObjectToPurge(const UObject* Object)
+	{
+		ObjectsToPurge.Add(Object);
+	}
+
+	virtual bool IsIgnoringArchetypeRef() const override
+	{
+		return false;
+	}
+
+	virtual bool IsIgnoringTransient() const override
+	{
+		return false;
+	}
+
+protected:
+	virtual void HandleObjectReference(UObject*& Object, const UObject* ReferencingObject, const UProperty* ReferencingProperty) override
+	{
+		ConditionalPurgeObject(Object);
+	}
+
+	virtual void HandleObjectReferences(UObject** InObjects, const int32 ObjectNum, const UObject* InReferencingObject, const UProperty* InReferencingProperty) override
+	{
+		for (int32 ObjectIndex = 0; ObjectIndex < ObjectNum; ++ObjectIndex)
+		{
+			UObject*& Object = InObjects[ObjectIndex];
+			ConditionalPurgeObject(Object);
+		}
+	}
+
+	void ConditionalPurgeObject(UObject*& Object)
+	{
+		if (ObjectsToPurge.Contains(Object))
+		{
+			Object = nullptr;
+		}
+	}
+
+	TSet<const UObject*> ObjectsToPurge;
+};
