Adding conditionals to guard against level editor tab manager and modules not being valid during tab close.

diff --git a/Engine/Source/Editor/Blutility/Private/EditorUtilityWidgetBlueprint.cpp b/Engine/Source/Editor/Blutility/Private/EditorUtilityWidgetBlueprint.cpp
--- a/Engine/Source/Editor/Blutility/Private/EditorUtilityWidgetBlueprint.cpp
+++ b/Engine/Source/Editor/Blutility/Private/EditorUtilityWidgetBlueprint.cpp
@@ -26,11 +26,20 @@ void UEditorUtilityWidgetBlueprint::BeginDestroy()
 	if (!GIsRequestingExit)
 	{
 		IBlutilityModule* BlutilityModule = FModuleManager::GetModulePtr<IBlutilityModule>("Blutility");
-		BlutilityModule->RemoveLoadedScriptUI(this);
-
-		FLevelEditorModule& LevelEditorModule = FModuleManager::GetModuleChecked<FLevelEditorModule>(TEXT("LevelEditor"));
-		TSharedPtr<FTabManager> LevelEditorTabManager = LevelEditorModule.GetLevelEditorTabManager();
-		LevelEditorTabManager->UnregisterTabSpawner(RegistrationName);
+		if (BlutilityModule)
+		{
+			BlutilityModule->RemoveLoadedScriptUI(this);
+		}
+
+		FLevelEditorModule* LevelEditorModule = FModuleManager::GetModulePtr<FLevelEditorModule>(TEXT("LevelEditor"));
+		if (LevelEditorModule)
+		{
+			TSharedPtr<FTabManager> LevelEditorTabManager = LevelEditorModule->GetLevelEditorTabManager();
+			if (LevelEditorTabManager.IsValid())
+			{
+				LevelEditorTabManager->UnregisterTabSpawner(RegistrationName);
+			}
+		}
 	}
 
 	Super::BeginDestroy();
