Added UEditorUtilityLibrary::GetSelectedAssetData to access useful information without loading the asset (path, name, class, package, etc.)

diff --git a/Engine/Source/Editor/Blutility/Private/EditorUtilityLibrary.cpp b/Engine/Source/Editor/Blutility/Private/EditorUtilityLibrary.cpp
--- a/Engine/Source/Editor/Blutility/Private/EditorUtilityLibrary.cpp
+++ b/Engine/Source/Editor/Blutility/Private/EditorUtilityLibrary.cpp
@@ -75,6 +75,15 @@ TArray<UObject*> UEditorUtilityLibrary::GetSelectedAssets()
 	return Result;
 }
 
+TArray<FAssetData> UEditorUtilityLibrary::GetSelectedAssetData()
+{
+	FContentBrowserModule& ContentBrowserModule = FModuleManager::LoadModuleChecked<FContentBrowserModule>("ContentBrowser");
+	TArray<FAssetData> SelectedAssets;
+	ContentBrowserModule.Get().GetSelectedAssets(SelectedAssets);
+
+	return SelectedAssets;
+}
+
 void UEditorUtilityLibrary::RenameAsset(UObject* Asset, const FString& NewName)
 {
 	FAssetToolsModule& AssetToolsModule = FModuleManager::LoadModuleChecked<FAssetToolsModule>("AssetTools");
diff --git a/Engine/Source/Editor/Blutility/Public/EditorUtilityLibrary.h b/Engine/Source/Editor/Blutility/Public/EditorUtilityLibrary.h
--- a/Engine/Source/Editor/Blutility/Public/EditorUtilityLibrary.h
+++ b/Engine/Source/Editor/Blutility/Public/EditorUtilityLibrary.h
@@ -29,6 +29,10 @@ public:
 	UFUNCTION(BlueprintCallable, Category = "Development|Editor")
 	static TArray<UObject*> GetSelectedAssets();
 
+	// Gets the set of currently selected asset data
+	UFUNCTION(BlueprintCallable, Category = "Development|Editor")
+	static TArray<FAssetData> GetSelectedAssetData();
+
 	// Renames an asset (cannot move folders)
 	UFUNCTION(BlueprintCallable, Category = "Development|Editor")
 	static void RenameAsset(UObject* Asset, const FString& NewName);
