- When destroying EditorUtilityWidgetBlueprint, it now checks if it has been registered to process unregistration

diff --git a/Engine/Source/Editor/Blutility/Private/EditorUtilityWidgetBlueprint.cpp b/Engine/Source/Editor/Blutility/Private/EditorUtilityWidgetBlueprint.cpp
--- a/Engine/Source/Editor/Blutility/Private/EditorUtilityWidgetBlueprint.cpp
+++ b/Engine/Source/Editor/Blutility/Private/EditorUtilityWidgetBlueprint.cpp
@@ -22,8 +22,8 @@ UEditorUtilityWidgetBlueprint::UEditorUtilityWidgetBlueprint(const FObjectInitia
 
 void UEditorUtilityWidgetBlueprint::BeginDestroy()
 {
-	// prevent the cleanup script from running on editor shutdown
-	if (!GIsRequestingExit)
+	// Only cleanup script if it has been registered and we're not shutdowning editor
+	if (!GIsRequestingExit && RegistrationName != NAME_None)
 	{
 		IBlutilityModule* BlutilityModule = FModuleManager::GetModulePtr<IBlutilityModule>("Blutility");
 		if (BlutilityModule)
