Fix for Editor Utility internals keeping a level reference alive

diff --git a/Engine/Source/Editor/Blutility/Private/EditorUtilityWidgetBlueprint.cpp b/Engine/Source/Editor/Blutility/Private/EditorUtilityWidgetBlueprint.cpp
--- a/Engine/Source/Editor/Blutility/Private/EditorUtilityWidgetBlueprint.cpp
+++ b/Engine/Source/Editor/Blutility/Private/EditorUtilityWidgetBlueprint.cpp
@@ -101,11 +101,12 @@ void UEditorUtilityWidgetBlueprint::ChangeTabWorld(UWorld* World, EMapChangeType
 {
 	if (MapChangeType == EMapChangeType::TearDownWorld)
 	{
-		CreatedUMGWidget = nullptr;
 		if (CreatedTab.IsValid())
 		{
 			CreatedTab.Pin()->SetContent(SNullWidget::NullWidget);
 		}
+		CreatedUMGWidget->Rename(nullptr, GetTransientPackage());
+		CreatedUMGWidget = nullptr;
 	}
 	else if (MapChangeType != EMapChangeType::SaveMap)
 	{
