Adding nullptr check for edge case when creating new levels
#rnx

diff --git a/Engine/Source/Editor/Blutility/Private/EditorUtilityWidgetBlueprint.cpp b/Engine/Source/Editor/Blutility/Private/EditorUtilityWidgetBlueprint.cpp
--- a/Engine/Source/Editor/Blutility/Private/EditorUtilityWidgetBlueprint.cpp
+++ b/Engine/Source/Editor/Blutility/Private/EditorUtilityWidgetBlueprint.cpp
@@ -105,7 +105,10 @@ void UEditorUtilityWidgetBlueprint::ChangeTabWorld(UWorld* World, EMapChangeType
 		{
 			CreatedTab.Pin()->SetContent(SNullWidget::NullWidget);
 		}
-		CreatedUMGWidget->Rename(nullptr, GetTransientPackage());
+		if (CreatedUMGWidget)
+		{
+			CreatedUMGWidget->Rename(nullptr, GetTransientPackage());
+		}
 		CreatedUMGWidget = nullptr;
 	}
 	else if (MapChangeType != EMapChangeType::SaveMap)
