Editor: Added support for content browser extension Blutilities to target/filter against Blueprint classes instead of regular asset types (by returning true for IsActionForBlueprints), and exposed UEditorUtilityLibrary::GetSelectedBlueprintClasses

diff --git a/Engine/Source/Editor/Blutility/Classes/AssetActionUtility.h b/Engine/Source/Editor/Blutility/Classes/AssetActionUtility.h
--- a/Engine/Source/Editor/Blutility/Classes/AssetActionUtility.h
+++ b/Engine/Source/Editor/Blutility/Classes/AssetActionUtility.h
@@ -12,7 +12,7 @@
 /** 
  * Base class for all asset action-related utilities
  * Any functions/events that are exposed on derived classes that have the correct signature will be
- * included as menu options when right-clicking on a group of assets int he content browser.
+ * included as menu options when right-clicking on a group of assets in the content browser.
  */
 UCLASS(Abstract, hideCategories=(Object), Blueprintable)
 class BLUTILITY_API UAssetActionUtility : public UEditorUtilityObject, public IEditorUtilityExtension
@@ -20,7 +20,14 @@ class BLUTILITY_API UAssetActionUtility : public UEditorUtilityObject, public IE
 	GENERATED_BODY()
 
 public:
-	/** Return the class that this asset action supports */
+	/** Return the class that this asset action supports (if not implemented, it will show up for all asset types) */
 	UFUNCTION(BlueprintPure, BlueprintImplementableEvent, Category="Assets")
 	UClass* GetSupportedClass() const;
+
+	/**
+	 * Returns whether or not this action is designed to work specifically on Blueprints (true) or on all assets (false).
+	 * If true, GetSupportedClass() is treated as a filter against the Parent Class of selected Blueprint assets
+	 */
+	UFUNCTION(BlueprintPure, BlueprintImplementableEvent, Category = "Assets")
+	bool IsActionForBlueprints() const;
 };
\ No newline at end of file
diff --git a/Engine/Source/Editor/Blutility/Private/BlutilityContentBrowserExtensions.cpp b/Engine/Source/Editor/Blutility/Private/BlutilityContentBrowserExtensions.cpp
--- a/Engine/Source/Editor/Blutility/Private/BlutilityContentBrowserExtensions.cpp
+++ b/Engine/Source/Editor/Blutility/Private/BlutilityContentBrowserExtensions.cpp
@@ -53,8 +53,30 @@ public:
 							{
 								if(UAssetActionUtility* DefaultObject = Cast<UAssetActionUtility>(BPClass->GetDefaultObject()))
 								{
+									bool bIsActionForBlueprints = DefaultObject->IsActionForBlueprints();
 									UClass* SupportedClass = DefaultObject->GetSupportedClass();
-									if(SupportedClass == nullptr || (SupportedClass && Asset.GetClass()->IsChildOf(SupportedClass)))
+									
+									bool bPassesClassFilter = false;
+									if(bIsActionForBlueprints)
+									{
+										if(UBlueprint* AssetAsBlueprint = Cast<UBlueprint>(Asset.GetAsset()))
+										{
+											// It's a blueprint, but is it the right kind?
+											bPassesClassFilter = (SupportedClass == nullptr || (SupportedClass && AssetAsBlueprint->ParentClass && AssetAsBlueprint->ParentClass->IsChildOf(SupportedClass)));
+										}
+										else
+										{
+											// Not a blueprint
+											bPassesClassFilter = false;
+										}
+									}
+									else
+									{
+										// Is the asset the right kind?
+										bPassesClassFilter = (SupportedClass == nullptr || (SupportedClass && Asset.GetClass()->IsChildOf(SupportedClass)));
+									}
+
+									if(bPassesClassFilter)
 									{
 										SupportedUtils.AddUnique(DefaultObject);
 									}
diff --git a/Engine/Source/Editor/Blutility/Private/EditorUtilityLibrary.cpp b/Engine/Source/Editor/Blutility/Private/EditorUtilityLibrary.cpp
--- a/Engine/Source/Editor/Blutility/Private/EditorUtilityLibrary.cpp
+++ b/Engine/Source/Editor/Blutility/Private/EditorUtilityLibrary.cpp
@@ -75,6 +75,25 @@ TArray<UObject*> UEditorUtilityLibrary::GetSelectedAssets()
 	return Result;
 }
 
+TArray<UClass*> UEditorUtilityLibrary::GetSelectedBlueprintClasses()
+{
+	//@TODO: Blocking load, no slow dialog
+	FContentBrowserModule& ContentBrowserModule = FModuleManager::LoadModuleChecked<FContentBrowserModule>("ContentBrowser");
+	TArray<FAssetData> SelectedAssets;
+	ContentBrowserModule.Get().GetSelectedAssets(SelectedAssets);
+
+	TArray<UClass*> Result;
+	for (FAssetData& AssetData : SelectedAssets)
+	{
+		if (UBlueprint* Blueprint = Cast<UBlueprint>(AssetData.GetAsset()))
+		{
+			Result.Add(Blueprint->GeneratedClass);
+		}
+	}
+
+	return Result;
+}
+
 TArray<FAssetData> UEditorUtilityLibrary::GetSelectedAssetData()
 {
 	FContentBrowserModule& ContentBrowserModule = FModuleManager::LoadModuleChecked<FContentBrowserModule>("ContentBrowser");
diff --git a/Engine/Source/Editor/Blutility/Public/EditorUtilityLibrary.h b/Engine/Source/Editor/Blutility/Public/EditorUtilityLibrary.h
--- a/Engine/Source/Editor/Blutility/Public/EditorUtilityLibrary.h
+++ b/Engine/Source/Editor/Blutility/Public/EditorUtilityLibrary.h
@@ -30,6 +30,10 @@ public:
 	UFUNCTION(BlueprintCallable, Category = "Development|Editor")
 	static TArray<UObject*> GetSelectedAssets();
 
+	// Gets the set of currently selected classes
+	UFUNCTION(BlueprintCallable, Category = "Development|Editor")
+	static TArray<UClass*> GetSelectedBlueprintClasses();
+
 	// Gets the set of currently selected asset data
 	UFUNCTION(BlueprintCallable, Category = "Development|Editor")
 	static TArray<FAssetData> GetSelectedAssetData();
