Merging //UE4/Dev-Main to Dev-VR (//UE4/Dev-VR)

Merging CL-6398973 (and previous change to same file 6398568) to fix build error

diff --git a/Engine/Source/Editor/Blutility/Private/EditorUtilitySubsystem.cpp b/Engine/Source/Editor/Blutility/Private/EditorUtilitySubsystem.cpp
--- a/Engine/Source/Editor/Blutility/Private/EditorUtilitySubsystem.cpp
+++ b/Engine/Source/Editor/Blutility/Private/EditorUtilitySubsystem.cpp
@@ -3,6 +3,7 @@
 #include "EditorUtilitySubsystem.h"
 #include "EditorUtilityCommon.h"
 #include "Interfaces/IMainFrameModule.h"
+#include "Engine/Blueprint.h"
 
 UEditorUtilitySubsystem::UEditorUtilitySubsystem() :
 	UEditorSubsystem()
@@ -55,7 +56,7 @@ bool UEditorUtilitySubsystem::TryRun(UObject* Asset)
 {
 	if (!Asset || Asset->IsPendingKillOrUnreachable())
 	{
-		UE_LOG(LogEditorUtilityBlueprint, Warning, TEXT("Could not run: %s"), *Asset->GetPathName());
+		UE_LOG(LogEditorUtilityBlueprint, Warning, TEXT("Could not run: %s"), Asset ? *Asset->GetPathName() : TEXT("None"));
 		return false;
 	}
 
