Fix crash when running editor utility function that contained parameters or return values

diff --git a/Engine/Source/Editor/Blutility/Private/EditorUtilitySubsystem.cpp b/Engine/Source/Editor/Blutility/Private/EditorUtilitySubsystem.cpp
--- a/Engine/Source/Editor/Blutility/Private/EditorUtilitySubsystem.cpp
+++ b/Engine/Source/Editor/Blutility/Private/EditorUtilitySubsystem.cpp
@@ -73,15 +73,22 @@ bool UEditorUtilitySubsystem::TryRun(UObject* Asset)
 	}
 
 	static const FName RunFunctionName("Run");
-	UFunction* RunFunction = ObjectClass->FindFunctionByName(RunFunctionName);
-	if (RunFunction)
+	if (UFunction* RunFunction = ObjectClass->FindFunctionByName(RunFunctionName))
 	{
-		UObject* Instance = NewObject<UObject>(this, ObjectClass);
-		ObjectInstances.Add(Asset, Instance);
+		if (RunFunction->NumParms == 0)
+		{
+			UObject* Instance = NewObject<UObject>(this, ObjectClass);
+			ObjectInstances.Add(Asset, Instance);
 
-		FEditorScriptExecutionGuard ScriptGuard;
-		Instance->ProcessEvent(RunFunction, nullptr);
-		return true;
+			FScopedTransaction Transaction( NSLOCTEXT("UnrealEd", "BlutilityAction", "Blutility Action") );
+			FEditorScriptExecutionGuard ScriptGuard;
+			Instance->ProcessEvent(RunFunction, nullptr);
+			return true;
+		}
+		else
+		{
+			UE_LOG(LogEditorUtilityBlueprint, Warning, TEXT("'Run' Function cannot have parameters or return values: %s"), *Asset->GetPathName());
+		}
 	}
 	else
 	{
