additional stylus module compile fixes
#rnx

diff --git a/Engine/Plugins/Editor/StylusInput/Source/StylusInput/Private/SStylusInputDebugWidget.cpp b/Engine/Plugins/Editor/StylusInput/Source/StylusInput/Private/SStylusInputDebugWidget.cpp
--- a/Engine/Plugins/Editor/StylusInput/Source/StylusInput/Private/SStylusInputDebugWidget.cpp
+++ b/Engine/Plugins/Editor/StylusInput/Source/StylusInput/Private/SStylusInputDebugWidget.cpp
@@ -218,4 +218,7 @@ ECheckBoxState SStylusInputDebugWidget::IsTouching() const
 ECheckBoxState SStylusInputDebugWidget::IsInverted() const
 {
 	return State.IsStylusInverted() ? ECheckBoxState::Checked : ECheckBoxState::Unchecked;
-}
\ No newline at end of file
+}
+
+
+#undef LOCTEXT_NAMESPACE
\ No newline at end of file
diff --git a/Engine/Plugins/Editor/StylusInput/Source/StylusInput/Private/StylusInputModule.cpp b/Engine/Plugins/Editor/StylusInput/Source/StylusInput/Private/StylusInputModule.cpp
--- a/Engine/Plugins/Editor/StylusInput/Source/StylusInput/Private/StylusInputModule.cpp
+++ b/Engine/Plugins/Editor/StylusInput/Source/StylusInput/Private/StylusInputModule.cpp
@@ -14,7 +14,6 @@
 
 #define LOCTEXT_NAMESPACE "FStylusInputModule"
 
-DEFINE_LOG_CATEGORY_STATIC(LogStylusInput, Log, All);
 
 static const FName StylusInputDebugTabName = FName("StylusInputDebug");
 
diff --git a/Engine/Plugins/Editor/StylusInput/Source/StylusInput/Private/WindowsRealTimeStylusPlugin.h b/Engine/Plugins/Editor/StylusInput/Source/StylusInput/Private/WindowsRealTimeStylusPlugin.h
--- a/Engine/Plugins/Editor/StylusInput/Source/StylusInput/Private/WindowsRealTimeStylusPlugin.h
+++ b/Engine/Plugins/Editor/StylusInput/Source/StylusInput/Private/WindowsRealTimeStylusPlugin.h
@@ -2,6 +2,7 @@
 
 #pragma once
 
+#include "CoreMinimal.h"
 #include "Windows/WindowsHWrapper.h"
 
 #include "Windows/AllowWindowsPlatformTypes.h"
diff --git a/Engine/Plugins/Editor/StylusInput/Source/StylusInput/Private/WindowsStylusInputInterface.cpp b/Engine/Plugins/Editor/StylusInput/Source/StylusInput/Private/WindowsStylusInputInterface.cpp
--- a/Engine/Plugins/Editor/StylusInput/Source/StylusInput/Private/WindowsStylusInputInterface.cpp
+++ b/Engine/Plugins/Editor/StylusInput/Source/StylusInput/Private/WindowsStylusInputInterface.cpp
@@ -4,8 +4,6 @@
 #include "WindowsRealTimeStylusPlugin.h"
 #include "Interfaces/IMainFrameModule.h"
 
-DEFINE_LOG_CATEGORY_STATIC(LogStylusInput, Log, All);
-
 class FWindowsStylusInputInterfaceImpl
 {
 public:
diff --git a/Engine/Plugins/Editor/StylusInput/Source/StylusInput/Public/IStylusInputModule.h b/Engine/Plugins/Editor/StylusInput/Source/StylusInput/Public/IStylusInputModule.h
--- a/Engine/Plugins/Editor/StylusInput/Source/StylusInput/Public/IStylusInputModule.h
+++ b/Engine/Plugins/Editor/StylusInput/Source/StylusInput/Public/IStylusInputModule.h
@@ -13,6 +13,8 @@
 
 #include "IStylusInputModule.generated.h"
 
+DEFINE_LOG_CATEGORY_STATIC(LogStylusInput, Log, All);
+
 /**
  * Module to handle Wacom-style tablet input using styluses.
  */
diff --git a/Engine/Source/Editor/Blutility/Private/EditorUtilitySubsystem.cpp b/Engine/Source/Editor/Blutility/Private/EditorUtilitySubsystem.cpp
--- a/Engine/Source/Editor/Blutility/Private/EditorUtilitySubsystem.cpp
+++ b/Engine/Source/Editor/Blutility/Private/EditorUtilitySubsystem.cpp
@@ -4,6 +4,7 @@
 #include "EditorUtilityCommon.h"
 #include "Interfaces/IMainFrameModule.h"
 #include "Engine/Blueprint.h"
+#include "ScopedTransaction.h"
 
 UEditorUtilitySubsystem::UEditorUtilitySubsystem() :
 	UEditorSubsystem()
@@ -73,22 +74,15 @@ bool UEditorUtilitySubsystem::TryRun(UObject* Asset)
 	}
 
 	static const FName RunFunctionName("Run");
-	if (UFunction* RunFunction = ObjectClass->FindFunctionByName(RunFunctionName))
+	UFunction* RunFunction = ObjectClass->FindFunctionByName(RunFunctionName);
+	if (RunFunction)
 	{
-		if (RunFunction->NumParms == 0)
-		{
-			UObject* Instance = NewObject<UObject>(this, ObjectClass);
-			ObjectInstances.Add(Asset, Instance);
+		UObject* Instance = NewObject<UObject>(this, ObjectClass);
+		ObjectInstances.Add(Asset, Instance);
 
-			FScopedTransaction Transaction( NSLOCTEXT("UnrealEd", "BlutilityAction", "Blutility Action") );
-			FEditorScriptExecutionGuard ScriptGuard;
-			Instance->ProcessEvent(RunFunction, nullptr);
-			return true;
-		}
-		else
-		{
-			UE_LOG(LogEditorUtilityBlueprint, Warning, TEXT("'Run' Function cannot have parameters or return values: %s"), *Asset->GetPathName());
-		}
+		FEditorScriptExecutionGuard ScriptGuard;
+		Instance->ProcessEvent(RunFunction, nullptr);
+		return true;
 	}
 	else
 	{
diff --git a/Engine/Source/Runtime/Core/Public/Windows/COMPointer.h b/Engine/Source/Runtime/Core/Public/Windows/COMPointer.h
--- a/Engine/Source/Runtime/Core/Public/Windows/COMPointer.h
+++ b/Engine/Source/Runtime/Core/Public/Windows/COMPointer.h
@@ -3,6 +3,7 @@
 #pragma once
 
 #include <Unknwn.h>
+#include "Misc/AssertionMacros.h"
 
 
 /**
