Don't reinstance the UMG widget for Editor Utility Widgets unless it has been destroyed somehow

diff --git a/Engine/Source/Editor/Blutility/Classes/EditorUtilityWidgetBlueprint.h b/Engine/Source/Editor/Blutility/Classes/EditorUtilityWidgetBlueprint.h
--- a/Engine/Source/Editor/Blutility/Classes/EditorUtilityWidgetBlueprint.h
+++ b/Engine/Source/Editor/Blutility/Classes/EditorUtilityWidgetBlueprint.h
@@ -31,7 +31,7 @@ public:
 	TSharedRef<SWidget> CreateUtilityWidget();
 
 	/** Recreate the tab's content on recompile */
-	void RegenerateCreatedTab(UBlueprint* RecompiledBlueprint);
+	void RegenerateCreatedTab();
 	
 	void UpdateRespawnListIfNeeded(TSharedRef<SDockTab> TabBeingClosed);
 
@@ -60,6 +60,7 @@ private:
 
 	TWeakPtr<SDockTab> CreatedTab;
 
+	UPROPERTY()
 	UEditorUtilityWidget* CreatedUMGWidget;
 
 };
diff --git a/Engine/Source/Editor/Blutility/Private/EditorUtilityWidgetBlueprint.cpp b/Engine/Source/Editor/Blutility/Private/EditorUtilityWidgetBlueprint.cpp
--- a/Engine/Source/Editor/Blutility/Private/EditorUtilityWidgetBlueprint.cpp
+++ b/Engine/Source/Editor/Blutility/Private/EditorUtilityWidgetBlueprint.cpp
@@ -56,33 +56,38 @@ TSharedRef<SDockTab> UEditorUtilityWidgetBlueprint::SpawnEditorUITab(const FSpaw
 	SpawnedTab->SetOnTabClosed(SDockTab::FOnTabClosedCallback::CreateUObject(this, &UEditorUtilityWidgetBlueprint::UpdateRespawnListIfNeeded));
 	CreatedTab = SpawnedTab;
 	
-	OnCompiled().AddUObject(this, &UEditorUtilityWidgetBlueprint::RegenerateCreatedTab);
+	GEditor->OnBlueprintReinstanced().AddUObject(this, &UEditorUtilityWidgetBlueprint::RegenerateCreatedTab);
 
 	return SpawnedTab;
 }
 
 TSharedRef<SWidget> UEditorUtilityWidgetBlueprint::CreateUtilityWidget()
 {
+	TSharedRef<SWidget> TabWidget = SNullWidget::NullWidget;
+
 	UClass* BlueprintClass = GeneratedClass;
 	TSubclassOf<UEditorUtilityWidget> WidgetClass = BlueprintClass;
 	UWorld* World = GEditor->GetEditorWorldContext().World();
 	check(World);
-	CreatedUMGWidget = CreateWidget<UEditorUtilityWidget>(World, WidgetClass);
-	TSharedRef<SWidget> TabWidget = SNullWidget::NullWidget;
+
+	if (!CreatedUMGWidget)
+	{
+		CreatedUMGWidget = CreateWidget<UEditorUtilityWidget>(World, WidgetClass);
+	}
+
 	if (CreatedUMGWidget)
 	{
-		TSharedRef<SWidget> CreatedSlateWidget = CreatedUMGWidget->TakeWidget();
 		TabWidget = SNew(SVerticalBox)
 			+ SVerticalBox::Slot()
 			.HAlign(HAlign_Fill)
 			[
-				CreatedSlateWidget
+				CreatedUMGWidget->TakeWidget()
 			];
 	}
 	return TabWidget;
 }
 
-void UEditorUtilityWidgetBlueprint::RegenerateCreatedTab(UBlueprint* RecompiledBlueprint)
+void UEditorUtilityWidgetBlueprint::RegenerateCreatedTab()
 {
 	if (CreatedTab.IsValid())
 	{
