Fixing world check after updating how Editor Utility Widget UI refreshes

diff --git a/Engine/Source/Editor/Blutility/Private/EditorUtilityWidgetBlueprint.cpp b/Engine/Source/Editor/Blutility/Private/EditorUtilityWidgetBlueprint.cpp
--- a/Engine/Source/Editor/Blutility/Private/EditorUtilityWidgetBlueprint.cpp
+++ b/Engine/Source/Editor/Blutility/Private/EditorUtilityWidgetBlueprint.cpp
@@ -68,9 +68,7 @@ TSharedRef<SWidget> UEditorUtilityWidgetBlueprint::CreateUtilityWidget()
 	UClass* BlueprintClass = GeneratedClass;
 	TSubclassOf<UEditorUtilityWidget> WidgetClass = BlueprintClass;
 	UWorld* World = GEditor->GetEditorWorldContext().World();
-	check(World);
-
-	if (!CreatedUMGWidget)
+	if (!CreatedUMGWidget && World)
 	{
 		CreatedUMGWidget = CreateWidget<UEditorUtilityWidget>(World, WidgetClass);
 	}
