Adding BP Diff action to editor utility widgets

diff --git a/Engine/Source/Editor/Blutility/Private/AssetTypeActions_EditorUtilityWidgetBlueprint.cpp b/Engine/Source/Editor/Blutility/Private/AssetTypeActions_EditorUtilityWidgetBlueprint.cpp
--- a/Engine/Source/Editor/Blutility/Private/AssetTypeActions_EditorUtilityWidgetBlueprint.cpp
+++ b/Engine/Source/Editor/Blutility/Private/AssetTypeActions_EditorUtilityWidgetBlueprint.cpp
@@ -88,6 +88,26 @@ uint32 FAssetTypeActions_EditorUtilityWidgetBlueprint::GetCategories()
 	return BlutilityModule->GetAssetCategory();
 }
 
+void FAssetTypeActions_EditorUtilityWidgetBlueprint::PerformAssetDiff(UObject* Asset1, UObject* Asset2, const struct FRevisionInfo& OldRevision, const struct FRevisionInfo& NewRevision) const
+{
+	UBlueprint* OldBlueprint = CastChecked<UBlueprint>(Asset1);
+	UBlueprint* NewBlueprint = CastChecked<UBlueprint>(Asset2);
+
+	// sometimes we're comparing different revisions of one single asset (other 
+	// times we're comparing two completely separate assets altogether)
+	bool bIsSingleAsset = (NewBlueprint->GetName() == OldBlueprint->GetName());
+
+	FText WindowTitle = LOCTEXT("NamelessWidgetBlueprintDiff", "Editor Utility Widget Blueprint Diff");
+	// if we're diffing one asset against itself 
+	if (bIsSingleAsset)
+	{
+		// identify the assumed single asset in the window's title
+		WindowTitle = FText::Format(LOCTEXT("WidgetBlueprintDiff", "{0} - Editor Utility Widget Blueprint Diff"), FText::FromString(NewBlueprint->GetName()));
+	}
+
+	SBlueprintDiff::CreateDiffWindow(WindowTitle, OldBlueprint, NewBlueprint, OldRevision, NewRevision);
+}
+
 void FAssetTypeActions_EditorUtilityWidgetBlueprint::ExecuteRun(FWeakBlueprintPointerArray InObjects)
 {
 	for (auto ObjIt = InObjects.CreateConstIterator(); ObjIt; ++ObjIt)
diff --git a/Engine/Source/Editor/Blutility/Private/AssetTypeActions_EditorUtilityWidgetBlueprint.h b/Engine/Source/Editor/Blutility/Private/AssetTypeActions_EditorUtilityWidgetBlueprint.h
--- a/Engine/Source/Editor/Blutility/Private/AssetTypeActions_EditorUtilityWidgetBlueprint.h
+++ b/Engine/Source/Editor/Blutility/Private/AssetTypeActions_EditorUtilityWidgetBlueprint.h
@@ -20,6 +20,7 @@ public:
 	virtual void GetActions(const TArray<UObject*>& InObjects, FMenuBuilder& MenuBuilder) override;
 	virtual void OpenAssetEditor(const TArray<UObject*>& InObjects, TSharedPtr<class IToolkitHost> EditWithinLevelEditor = TSharedPtr<IToolkitHost>()) override;
 	virtual uint32 GetCategories() override;
+	virtual void PerformAssetDiff(UObject* Asset1, UObject* Asset2, const struct FRevisionInfo& OldRevision, const struct FRevisionInfo& NewRevision) const override;
 	virtual bool CanLocalize() const override { return false; }
 	// End of IAssetTypeActions interface
 
