Adding missing include

diff --git a/Engine/Source/Editor/Blutility/Private/AssetTypeActions_EditorUtilityWidgetBlueprint.cpp b/Engine/Source/Editor/Blutility/Private/AssetTypeActions_EditorUtilityWidgetBlueprint.cpp
--- a/Engine/Source/Editor/Blutility/Private/AssetTypeActions_EditorUtilityWidgetBlueprint.cpp
+++ b/Engine/Source/Editor/Blutility/Private/AssetTypeActions_EditorUtilityWidgetBlueprint.cpp
@@ -18,6 +18,7 @@
 #include "Widgets/Docking/SDockTab.h"
 #include "Framework/Docking/TabManager.h"
 #include "IBlutilityModule.h"
+#include "SBlueprintDiff.h"
 
 #define LOCTEXT_NAMESPACE "AssetTypeActions"
 
