Allow 'Run' on Editor Utility Object's to be called by blueprints.

diff --git a/Engine/Source/Editor/Blutility/Classes/EditorUtilityObject.h b/Engine/Source/Editor/Blutility/Classes/EditorUtilityObject.h
--- a/Engine/Source/Editor/Blutility/Classes/EditorUtilityObject.h
+++ b/Engine/Source/Editor/Blutility/Classes/EditorUtilityObject.h
@@ -19,6 +19,6 @@ class BLUTILITY_API UEditorUtilityObject : public UObject
 	GENERATED_UCLASS_BODY()
 
 	// Standard function to execute
-	UFUNCTION(BlueprintImplementableEvent)
+	UFUNCTION(BlueprintCallable, BlueprintImplementableEvent, Category = "Editor")
 	void Run();
 };
