Allow 'Run' on Editor Utility Object's to be called by blueprints.

diff --git a/Engine/Source/Editor/Blutility/Classes/EditorUtilityActor.h b/Engine/Source/Editor/Blutility/Classes/EditorUtilityActor.h
--- a/Engine/Source/Editor/Blutility/Classes/EditorUtilityActor.h
+++ b/Engine/Source/Editor/Blutility/Classes/EditorUtilityActor.h
@@ -18,6 +18,6 @@ class BLUTILITY_API AEditorUtilityActor : public AActor
 	GENERATED_UCLASS_BODY()
 
 	// Standard function to execute
-	UFUNCTION(BlueprintImplementableEvent)
+	UFUNCTION(BlueprintCallable, BlueprintImplementableEvent, Category = "Editor")
 	void Run();
 };
diff --git a/Engine/Source/Editor/Blutility/Classes/EditorUtilityWidget.h b/Engine/Source/Editor/Blutility/Classes/EditorUtilityWidget.h
--- a/Engine/Source/Editor/Blutility/Classes/EditorUtilityWidget.h
+++ b/Engine/Source/Editor/Blutility/Classes/EditorUtilityWidget.h
@@ -23,7 +23,7 @@ class BLUTILITY_API UEditorUtilityWidget : public UUserWidget
 
 public:
 	// The default action called when the widget is invoked if bAutoRunDefaultAction=true (it is never called otherwise)
-	UFUNCTION(BlueprintImplementableEvent)
+	UFUNCTION(BlueprintCallable, BlueprintImplementableEvent, Category = "Editor")
 	void Run();
 
 	// Run the default action
