Copied root widget selection dialog for use in editor utility widgets

diff --git a/Engine/Source/Editor/Blutility/Classes/EditorUtilityWidgetBlueprintFactory.h b/Engine/Source/Editor/Blutility/Classes/EditorUtilityWidgetBlueprintFactory.h
--- a/Engine/Source/Editor/Blutility/Classes/EditorUtilityWidgetBlueprintFactory.h
+++ b/Engine/Source/Editor/Blutility/Classes/EditorUtilityWidgetBlueprintFactory.h
@@ -37,4 +37,5 @@ class UEditorUtilityWidgetBlueprintFactory : public UFactory
 	
 	// End of UFactory interface
 
+	UClass* RootWidgetClass;
 };
diff --git a/Engine/Source/Editor/Blutility/Private/EditorUtilityWidgetBlueprintFactory.cpp b/Engine/Source/Editor/Blutility/Private/EditorUtilityWidgetBlueprintFactory.cpp
--- a/Engine/Source/Editor/Blutility/Private/EditorUtilityWidgetBlueprintFactory.cpp
+++ b/Engine/Source/Editor/Blutility/Private/EditorUtilityWidgetBlueprintFactory.cpp
@@ -18,20 +18,32 @@
 #include "Components/CanvasPanel.h"
 #include "BaseWidgetBlueprint.h"
 #include "Blueprint/WidgetTree.h"
+#include "Components/HorizontalBox.h"
+#include "Components/VerticalBox.h"
+#include "Components/GridPanel.h"
+#include "UMGEditorProjectSettings.h"
+
+#define LOCTEXT_NAMESPACE "UEditorUtilityWidgetBlueprintFactory"
 
 class FEditorUtilityWidgetBlueprintFactoryFilter : public IClassViewerFilter
 {
 public:
-	TSet< const UClass* > AllowedChildOfClasses;
+	/** All children of these classes will be included unless filtered out by another setting. */
+	TSet< const UClass* > AllowedChildrenOfClasses;
+
+	/** Disallowed class flags. */
+	EClassFlags DisallowedClassFlags;
 
 	bool IsClassAllowed(const FClassViewerInitializationOptions& InInitOptions, const UClass* InClass, TSharedRef< FClassViewerFilterFuncs > InFilterFuncs ) override
 	{
-		return InFilterFuncs->IfInChildOfClassesSet(AllowedChildOfClasses, InClass) != EFilterReturn::Failed;
+		return !InClass->HasAnyClassFlags(DisallowedClassFlags)
+			&& InFilterFuncs->IfInChildOfClassesSet(AllowedChildrenOfClasses, InClass) != EFilterReturn::Failed;
 	}
 
 	virtual bool IsUnloadedClassAllowed(const FClassViewerInitializationOptions& InInitOptions, const TSharedRef< const IUnloadedBlueprintData > InUnloadedClassData, TSharedRef< FClassViewerFilterFuncs > InFilterFuncs) override
 	{
-		return InFilterFuncs->IfInChildOfClassesSet(AllowedChildOfClasses, InUnloadedClassData) != EFilterReturn::Failed;
+		return !InUnloadedClassData->HasAnyClassFlags(DisallowedClassFlags)
+			&& InFilterFuncs->IfInChildOfClassesSet(AllowedChildrenOfClasses, InUnloadedClassData) != EFilterReturn::Failed;
 	}
 };
 
@@ -49,6 +61,31 @@ UEditorUtilityWidgetBlueprintFactory::UEditorUtilityWidgetBlueprintFactory(const
 
 bool UEditorUtilityWidgetBlueprintFactory::ConfigureProperties()
 {
+	if (GetDefault<UUMGEditorProjectSettings>()->bUseWidgetTemplateSelector)
+	{
+		// Load the classviewer module to display a class picker
+		FClassViewerModule& ClassViewerModule = FModuleManager::LoadModuleChecked<FClassViewerModule>("ClassViewer");
+
+		// Fill in options
+		FClassViewerInitializationOptions Options;
+		Options.Mode = EClassViewerMode::ClassPicker;
+		Options.bShowNoneOption = true;
+
+		Options.ExtraPickerCommonClasses.Add(UHorizontalBox::StaticClass());
+		Options.ExtraPickerCommonClasses.Add(UVerticalBox::StaticClass());
+		Options.ExtraPickerCommonClasses.Add(UGridPanel::StaticClass());
+		Options.ExtraPickerCommonClasses.Add(UCanvasPanel::StaticClass());
+
+		TSharedPtr<FEditorUtilityWidgetBlueprintFactoryFilter> Filter = MakeShareable(new FEditorUtilityWidgetBlueprintFactoryFilter);
+		Options.ClassFilter = Filter;
+
+		Filter->DisallowedClassFlags = CLASS_Abstract | CLASS_Deprecated | CLASS_NewerVersionExists;
+		Filter->AllowedChildrenOfClasses.Add(UPanelWidget::StaticClass());
+
+		const FText TitleText = LOCTEXT("CreateWidgetBlueprint", "Pick Root Widget for New Editor Utility Widget");
+		return SClassPickerDialog::PickClass(TitleText, Options, RootWidgetClass, UPanelWidget::StaticClass());
+
+	}
 	return true;
 }
 
@@ -68,13 +105,21 @@ UObject* UEditorUtilityWidgetBlueprintFactory::FactoryCreateNew(UClass* Class, U
 	}
 	else
 	{
+		// If the root widget selection dialog is not enabled, use a canvas panel as the root by default
+		if (!GetDefault<UUMGEditorProjectSettings>()->bUseWidgetTemplateSelector)
+		{
+			RootWidgetClass = UCanvasPanel::StaticClass();
+		}
 		UEditorUtilityWidgetBlueprint* NewBP = CastChecked<UEditorUtilityWidgetBlueprint>(FKismetEditorUtilities::CreateBlueprint(ParentClass, InParent, Name, BlueprintType, UEditorUtilityWidgetBlueprint::StaticClass(), UWidgetBlueprintGeneratedClass::StaticClass(), NAME_None));
 
-		// Create a CanvasPanel to use as the default root widget
+		// Create the selected root widget
 		if (NewBP->WidgetTree->RootWidget == nullptr)
 		{
-			UWidget* Root = NewBP->WidgetTree->ConstructWidget<UCanvasPanel>(UCanvasPanel::StaticClass());
-			NewBP->WidgetTree->RootWidget = Root;
+			if (TSubclassOf<UPanelWidget> RootWidgetPanel = RootWidgetClass)
+			{
+				UWidget* Root = NewBP->WidgetTree->ConstructWidget<UWidget>(RootWidgetPanel);
+				NewBP->WidgetTree->RootWidget = Root;
+			}
 		}
 
 		return NewBP;
@@ -85,3 +130,5 @@ bool UEditorUtilityWidgetBlueprintFactory::CanCreateNew() const
 {
 	return true;
 }
+
+#undef LOCTEXT_NAMESPACE
\ No newline at end of file
