Updated editor utility widgets to reinstance on compile and move the old instance to the transient package so that GC can pick it up. This fixes an issue where OnConstruct would be called on the same instance each time the widget is recompiled.

diff --git a/Engine/Source/Editor/Blutility/Private/EditorUtilityWidgetBlueprint.cpp b/Engine/Source/Editor/Blutility/Private/EditorUtilityWidgetBlueprint.cpp
--- a/Engine/Source/Editor/Blutility/Private/EditorUtilityWidgetBlueprint.cpp
+++ b/Engine/Source/Editor/Blutility/Private/EditorUtilityWidgetBlueprint.cpp
@@ -71,8 +71,12 @@ TSharedRef<SWidget> UEditorUtilityWidgetBlueprint::CreateUtilityWidget()
 	UClass* BlueprintClass = GeneratedClass;
 	TSubclassOf<UEditorUtilityWidget> WidgetClass = BlueprintClass;
 	UWorld* World = GEditor->GetEditorWorldContext().World();
-	if (!CreatedUMGWidget && World)
+	if (World)
 	{
+		if (CreatedUMGWidget)
+		{
+			CreatedUMGWidget->Rename(nullptr, GetTransientPackage());
+		}
 		CreatedUMGWidget = CreateWidget<UEditorUtilityWidget>(World, WidgetClass);
 	}
 
