Adding more tab utility functions to Editor Utility Subsystem
#rnx

diff --git a/Engine/Source/Editor/Blutility/Private/EditorUtilitySubsystem.cpp b/Engine/Source/Editor/Blutility/Private/EditorUtilitySubsystem.cpp
--- a/Engine/Source/Editor/Blutility/Private/EditorUtilitySubsystem.cpp
+++ b/Engine/Source/Editor/Blutility/Private/EditorUtilitySubsystem.cpp
@@ -104,9 +104,22 @@ void UEditorUtilitySubsystem::ReleaseInstanceOfAsset(UObject* Asset)
 	ObjectInstances.Remove(Asset);
 }
 
-UEditorUtilityWidget* UEditorUtilitySubsystem::SpawnAndRegisterTab(UEditorUtilityWidgetBlueprint* InBlueprint)
+UEditorUtilityWidget* UEditorUtilitySubsystem::SpawnAndRegisterTabAndGetID(UEditorUtilityWidgetBlueprint* InBlueprint, FName& NewTabID)
 {
+	RegisterTabAndGetID(InBlueprint, NewTabID);
+	SpawnRegisteredTabByID(NewTabID);
+	return GetUtilityWidgetFromBlueprint(InBlueprint);
+}
+
+
+UEditorUtilityWidget* UEditorUtilitySubsystem::SpawnAndRegisterTab(class UEditorUtilityWidgetBlueprint* InBlueprint)
+{
+	FName InTabID;
+	return SpawnAndRegisterTabAndGetID(InBlueprint, InTabID);
+}
 
+void UEditorUtilitySubsystem::RegisterTabAndGetID(class UEditorUtilityWidgetBlueprint* InBlueprint, FName& NewTabID)
+{
 	if (InBlueprint && !IsRunningCommandlet())
 	{
 		FName RegistrationName = FName(*(InBlueprint->GetPathName() + LOCTEXT("ActiveTabSuffix", "_ActiveTab").ToString()));
@@ -122,12 +135,62 @@ UEditorUtilityWidget* UEditorUtilitySubsystem::SpawnAndRegisterTab(UEditorUtilit
 			InBlueprint->SetRegistrationName(RegistrationName);
 			BlutilityModule->AddLoadedScriptUI(InBlueprint);
 		}
-		TSharedRef<SDockTab> NewDockTab = LevelEditorTabManager->InvokeTab(RegistrationName);
-		return InBlueprint->GetCreatedWidget();
+		NewTabID = RegistrationName;
+	}
+
+}
+
+bool UEditorUtilitySubsystem::SpawnRegisteredTabByID(FName NewTabID)
+{
+	if (!IsRunningCommandlet())
+	{
+		FLevelEditorModule& LevelEditorModule = FModuleManager::GetModuleChecked<FLevelEditorModule>(TEXT("LevelEditor"));
+		TSharedPtr<FTabManager> LevelEditorTabManager = LevelEditorModule.GetLevelEditorTabManager();
+		if (!LevelEditorTabManager->HasTabSpawner(NewTabID))
+		{
+			TSharedRef<SDockTab> NewDockTab = LevelEditorTabManager->InvokeTab(NewTabID);
+			return true;
+		}
+	}
+	return false;
+}
+
+bool UEditorUtilitySubsystem::DoesTabExist(FName NewTabID)
+{
+	if (!IsRunningCommandlet())
+	{
+		FLevelEditorModule& LevelEditorModule = FModuleManager::GetModuleChecked<FLevelEditorModule>(TEXT("LevelEditor"));
+		TSharedPtr<FTabManager> LevelEditorTabManager = LevelEditorModule.GetLevelEditorTabManager();
+		TSharedPtr<SDockTab> FoundTab = LevelEditorTabManager->FindExistingLiveTab(NewTabID);
+		if (FoundTab)
+		{
+			return true;
+		}
+	}
+
+	return false;
+}
+
+bool UEditorUtilitySubsystem::CloseTabByID(FName NewTabID)
+{
+	if (!IsRunningCommandlet())
+	{
+		FLevelEditorModule& LevelEditorModule = FModuleManager::GetModuleChecked<FLevelEditorModule>(TEXT("LevelEditor"));
+		TSharedPtr<FTabManager> LevelEditorTabManager = LevelEditorModule.GetLevelEditorTabManager();
+		TSharedPtr<SDockTab> FoundTab = LevelEditorTabManager->FindExistingLiveTab(NewTabID);
+		if (FoundTab)
+		{
+			FoundTab->RequestCloseTab();
+			return true;
+		}
 	}
 
-	return nullptr;
+	return false;
 }
 
+UEditorUtilityWidget* UEditorUtilitySubsystem::FindUtilityWidgetFromBlueprint(class UEditorUtilityWidgetBlueprint* InBlueprint)
+{
+	return InBlueprint->GetCreatedWidget();
+}
 
 #undef LOCTEXT_NAMESPACE
diff --git a/Engine/Source/Editor/Blutility/Public/EditorUtilitySubsystem.h b/Engine/Source/Editor/Blutility/Public/EditorUtilitySubsystem.h
--- a/Engine/Source/Editor/Blutility/Public/EditorUtilitySubsystem.h
+++ b/Engine/Source/Editor/Blutility/Public/EditorUtilitySubsystem.h
@@ -39,9 +39,31 @@ public:
 	UFUNCTION(BlueprintCallable, Category = "Development|Editor")
 	bool TryRun(UObject* Asset);
 
+	UFUNCTION(BlueprintCallable, Category = "Development|Editor")
+	UEditorUtilityWidget* SpawnAndRegisterTabAndGetID(class UEditorUtilityWidgetBlueprint* InBlueprint, FName& NewTabID);
+
 	UFUNCTION(BlueprintCallable, Category = "Development|Editor")
 	UEditorUtilityWidget* SpawnAndRegisterTab(class UEditorUtilityWidgetBlueprint* InBlueprint);
 
+	UFUNCTION(BlueprintCallable, Category = "Development|Editor")
+	void RegisterTabAndGetID(class UEditorUtilityWidgetBlueprint* InBlueprint, FName& NewTabID);
+
+	/** Given an ID for a tab, try to find a tab spawner that matches, and then spawn a tab. Returns true if it was able to find a matching tab spawner */
+	UFUNCTION(BlueprintCallable, Category = "Development|Editor")
+	bool SpawnRegisteredTabByID(FName NewTabID);
+
+	/** Given an ID for a tab, try to find an existing tab. Returns true if it found a tab. */
+	UFUNCTION(BlueprintCallable, Category = "Development|Editor")
+	bool DoesTabExist(FName NewTabID);
+
+	/** Given an ID for a tab, try to find and close an existing tab. Returns true if it found a tab to close. */
+	UFUNCTION(BlueprintCallable, Category = "Development|Editor")
+	bool CloseTabByID(FName NewTabID);
+
+	/** Given an editor utility widget blueprint, get the widget it creates. This will return a null pointer if the widget is not currently in a tab.*/
+	UFUNCTION(BlueprintCallable, Category = "Development|Editor")
+	UEditorUtilityWidget* FindUtilityWidgetFromBlueprint(class UEditorUtilityWidgetBlueprint* InBlueprint);
+
 private:
 	
 	UPROPERTY()
