Fix for missed function call rename
#rnx

diff --git a/Engine/Source/Editor/Blutility/Private/EditorUtilitySubsystem.cpp b/Engine/Source/Editor/Blutility/Private/EditorUtilitySubsystem.cpp
--- a/Engine/Source/Editor/Blutility/Private/EditorUtilitySubsystem.cpp
+++ b/Engine/Source/Editor/Blutility/Private/EditorUtilitySubsystem.cpp
@@ -108,7 +108,7 @@ UEditorUtilityWidget* UEditorUtilitySubsystem::SpawnAndRegisterTabAndGetID(UEdit
 {
 	RegisterTabAndGetID(InBlueprint, NewTabID);
 	SpawnRegisteredTabByID(NewTabID);
-	return GetUtilityWidgetFromBlueprint(InBlueprint);
+	return FindUtilityWidgetFromBlueprint(InBlueprint);
 }
 
 
