Fix editor utility widgets not spawning tabs

diff --git a/Engine/Source/Editor/Blutility/Private/EditorUtilitySubsystem.cpp b/Engine/Source/Editor/Blutility/Private/EditorUtilitySubsystem.cpp
--- a/Engine/Source/Editor/Blutility/Private/EditorUtilitySubsystem.cpp
+++ b/Engine/Source/Editor/Blutility/Private/EditorUtilitySubsystem.cpp
@@ -146,7 +146,7 @@ bool UEditorUtilitySubsystem::SpawnRegisteredTabByID(FName NewTabID)
 	{
 		FLevelEditorModule& LevelEditorModule = FModuleManager::GetModuleChecked<FLevelEditorModule>(TEXT("LevelEditor"));
 		TSharedPtr<FTabManager> LevelEditorTabManager = LevelEditorModule.GetLevelEditorTabManager();
-		if (!LevelEditorTabManager->HasTabSpawner(NewTabID))
+		if (LevelEditorTabManager->HasTabSpawner(NewTabID))
 		{
 			TSharedRef<SDockTab> NewDockTab = LevelEditorTabManager->InvokeTab(NewTabID);
 			return true;
