EditorUtilityActor and EditorUtilityObject now show up as common classes when you create an EditorUtilityBlueprint

diff --git a/Engine/Source/Editor/Blutility/Private/EditorUtilityBlueprintFactory.cpp b/Engine/Source/Editor/Blutility/Private/EditorUtilityBlueprintFactory.cpp
--- a/Engine/Source/Editor/Blutility/Private/EditorUtilityBlueprintFactory.cpp
+++ b/Engine/Source/Editor/Blutility/Private/EditorUtilityBlueprintFactory.cpp
@@ -15,6 +15,8 @@
 #include "Kismet2/KismetEditorUtilities.h"
 #include "Kismet2/SClassPickerDialog.h"
 #include "EditorUtilityWidget.h"
+#include "EditorUtilityActor.h"
+#include "EditorUtilityObject.h"
 
 class FBlutilityBlueprintFactoryFilter : public IClassViewerFilter
 {
@@ -74,6 +76,9 @@ bool UEditorUtilityBlueprintFactory::ConfigureProperties()
 	Options.bShowUnloadedBlueprints = true;
 	Options.bEditorClassesOnly = true;
 
+	Options.ExtraPickerCommonClasses.Add(AEditorUtilityActor::StaticClass());
+	Options.ExtraPickerCommonClasses.Add(UEditorUtilityObject::StaticClass());
+
 	TSharedPtr< FBlutilityBlueprintFactoryFilter > Filter = MakeShareable(new FBlutilityBlueprintFactoryFilter);
 	Filter->DisallowedChildOfClasses.Add(UEditorUtilityWidget::StaticClass());
 	Options.ClassFilter = Filter;
