Marked UEditurUtilityWidgetBlueprint as BLUTILITY_API

diff --git a/Engine/Source/Editor/Blutility/Classes/EditorUtilityWidgetBlueprint.h b/Engine/Source/Editor/Blutility/Classes/EditorUtilityWidgetBlueprint.h
--- a/Engine/Source/Editor/Blutility/Classes/EditorUtilityWidgetBlueprint.h
+++ b/Engine/Source/Editor/Blutility/Classes/EditorUtilityWidgetBlueprint.h
@@ -19,7 +19,7 @@ enum class EAssetEditorCloseReason : uint8;
 enum class EMapChangeType : uint8;
 
 UCLASS()
-class UEditorUtilityWidgetBlueprint : public UWidgetBlueprint
+class BLUTILITY_API UEditorUtilityWidgetBlueprint : public UWidgetBlueprint
 {
 	GENERATED_UCLASS_BODY()
 
