Asset/actor action utilities improvements:
1. Added ability to define utilities natively, or in Python (UClasses are now harvested alongside blueprint assets)
2. Functions are now grouped according to their category and placed in a submenu
3. Added ability to automatically forward the object selection to the first parameter of the function. (AActor and FAssetData) Only sub-set of the selection that passes SupportedClass is forwarded.

diff --git a/Engine/Source/Editor/Blutility/Private/BlutilityContentBrowserExtensions.cpp b/Engine/Source/Editor/Blutility/Private/BlutilityContentBrowserExtensions.cpp
--- a/Engine/Source/Editor/Blutility/Private/BlutilityContentBrowserExtensions.cpp
+++ b/Engine/Source/Editor/Blutility/Private/BlutilityContentBrowserExtensions.cpp
@@ -22,6 +22,8 @@
 #include "BlueprintEditorModule.h"
 #include "BlutilityMenuExtensions.h"
 
+#include "IBlutilityModule.h"
+
 #define LOCTEXT_NAMESPACE "BlutilityContentBrowserExtensions"
 
 static FContentBrowserMenuExtender_SelectedAssets ContentBrowserExtenderDelegate;
@@ -35,16 +37,63 @@ public:
 		TSharedRef<FExtender> Extender(new FExtender());
 
 		// Run thru the assets to determine if any meet our criteria
+		TMap<IEditorUtilityExtension*, TSet<int32>> UtilityAndSelectionIndices;
 		TArray<IEditorUtilityExtension*> SupportedUtils;
+		TArray<FAssetData> SupportedAssets;
 		if (SelectedAssets.Num() > 0)
 		{
 			// Check blueprint utils (we need to load them to query their validity against these assets)
 			TArray<FAssetData> UtilAssets;
 			FBlutilityMenuExtensions::GetBlutilityClasses(UtilAssets, UAssetActionUtility::StaticClass()->GetFName());
-			if (UtilAssets.Num() > 0)
+
+			// Collect all UAssetActionUtility derived classes
+			TSet<UAssetActionUtility*> AssetClasses;
+			for (TObjectIterator<UClass> AssetActionClassIt; AssetActionClassIt; ++AssetActionClassIt)
+			{
+				if (AssetActionClassIt->IsChildOf(UAssetActionUtility::StaticClass()) && UAssetActionUtility::StaticClass()->GetFName() != AssetActionClassIt->GetFName() && AssetActionClassIt->ClassGeneratedBy == nullptr)
+				{
+					AssetClasses.Add(Cast<UAssetActionUtility>(AssetActionClassIt->GetDefaultObject()));
+				}
+			}
+
+			if (UtilAssets.Num() + AssetClasses.Num() > 0)
 			{
 				for (const FAssetData& Asset : SelectedAssets)
 				{
+					auto ProcessAssetAction = [&SupportedUtils, &SupportedAssets, &UtilityAndSelectionIndices, Asset](UAssetActionUtility* InAction)
+					{
+						bool bIsActionForBlueprints = InAction->IsActionForBlueprints();
+						UClass* SupportedClass = InAction->GetSupportedClass();
+
+						bool bPassesClassFilter = false;
+						if (bIsActionForBlueprints)
+						{
+							if (UBlueprint* AssetAsBlueprint = Cast<UBlueprint>(Asset.GetAsset()))
+							{
+								// It's a blueprint, but is it the right kind?
+								bPassesClassFilter = (SupportedClass == nullptr || (SupportedClass && AssetAsBlueprint->ParentClass && AssetAsBlueprint->ParentClass->IsChildOf(SupportedClass)));
+							}
+							else
+							{
+								// Not a blueprint
+								bPassesClassFilter = false;
+							}
+						}
+						else
+						{
+							// Is the asset the right kind?
+							bPassesClassFilter = (SupportedClass == nullptr || (SupportedClass && Asset.GetClass()->IsChildOf(SupportedClass)));
+						}
+
+						if (bPassesClassFilter)
+						{
+							SupportedUtils.AddUnique(InAction);
+							const int32 Index = SupportedAssets.AddUnique(Asset);
+							UtilityAndSelectionIndices.FindOrAdd(InAction).Add(Index);
+						}
+					};
+
+					// Process asset based utilities
 					for (const FAssetData& UtilAsset : UtilAssets)
 					{
 						if(UEditorUtilityBlueprint* Blueprint = Cast<UEditorUtilityBlueprint>(UtilAsset.GetAsset()))
@@ -53,37 +102,17 @@ public:
 							{
 								if(UAssetActionUtility* DefaultObject = Cast<UAssetActionUtility>(BPClass->GetDefaultObject()))
 								{
-									bool bIsActionForBlueprints = DefaultObject->IsActionForBlueprints();
-									UClass* SupportedClass = DefaultObject->GetSupportedClass();
-									
-									bool bPassesClassFilter = false;
-									if(bIsActionForBlueprints)
-									{
-										if(UBlueprint* AssetAsBlueprint = Cast<UBlueprint>(Asset.GetAsset()))
-										{
-											// It's a blueprint, but is it the right kind?
-											bPassesClassFilter = (SupportedClass == nullptr || (SupportedClass && AssetAsBlueprint->ParentClass && AssetAsBlueprint->ParentClass->IsChildOf(SupportedClass)));
-										}
-										else
-										{
-											// Not a blueprint
-											bPassesClassFilter = false;
-										}
-									}
-									else
-									{
-										// Is the asset the right kind?
-										bPassesClassFilter = (SupportedClass == nullptr || (SupportedClass && Asset.GetClass()->IsChildOf(SupportedClass)));
-									}
-
-									if(bPassesClassFilter)
-									{
-										SupportedUtils.AddUnique(DefaultObject);
-									}
+									ProcessAssetAction(DefaultObject);
 								}
 							}
 						}
 					}
+
+					// Process non-asset based utilities
+					for (UAssetActionUtility* Action : AssetClasses)
+					{
+						ProcessAssetAction(Action);
+					}
 				}
 			}
 		}
@@ -95,7 +124,7 @@ public:
 				"CommonAssetActions",
 				EExtensionHook::After,
 				nullptr,
-				FMenuExtensionDelegate::CreateStatic(&FBlutilityMenuExtensions::CreateBlutilityActionsMenu, SupportedUtils));
+				FMenuExtensionDelegate::CreateStatic(&FBlutilityMenuExtensions::CreateAssetBlutilityActionsMenu, UtilityAndSelectionIndices, MoveTemp(SupportedAssets)));
 		}
 
 		return Extender;
diff --git a/Engine/Source/Editor/Blutility/Private/BlutilityLevelEditorExtensions.cpp b/Engine/Source/Editor/Blutility/Private/BlutilityLevelEditorExtensions.cpp
--- a/Engine/Source/Editor/Blutility/Private/BlutilityLevelEditorExtensions.cpp
+++ b/Engine/Source/Editor/Blutility/Private/BlutilityLevelEditorExtensions.cpp
@@ -25,37 +25,72 @@ public:
 	static TSharedRef<FExtender> OnExtendLevelEditorActorContextMenu(const TSharedRef<FUICommandList> CommandList, const TArray<AActor*> SelectedActors)
 	{
 		TSharedRef<FExtender> Extender(new FExtender());
+			   
+		TMap<IEditorUtilityExtension*, TSet<int32>> UtilityAndSelectionIndices;
 
-		// Run thru the assets to determine if any meet our criteria
+		// Run thru the actors to determine if any meet our criteria
 		TArray<IEditorUtilityExtension*> SupportedUtils;
+		TArray<AActor*> SupportedActors;
 		if (SelectedActors.Num() > 0)
 		{
-			// Check blueprint utils (we need to load them to query their validity against these actors)
+			// Check blueprint utils (we need to load them to query their validity against these assets)
 			TArray<FAssetData> UtilAssets;
 			FBlutilityMenuExtensions::GetBlutilityClasses(UtilAssets, UActorActionUtility::StaticClass()->GetFName());
-			if (UtilAssets.Num() > 0)
+
+			// Collect all UActorActionUtility derived (non-blueprint) classes 
+			TSet<UActorActionUtility*> AssetClasses;
+			for (TObjectIterator<UClass> AssetActionClassIt; AssetActionClassIt; ++AssetActionClassIt)
+			{
+				if (AssetActionClassIt->IsChildOf(UActorActionUtility::StaticClass()) && UActorActionUtility::StaticClass()->GetFName() != AssetActionClassIt->GetFName() && AssetActionClassIt->ClassGeneratedBy == nullptr)
+				{
+					AssetClasses.Add(Cast<UActorActionUtility>(AssetActionClassIt->GetDefaultObject()));
+				}
+			}
+
+			if (UtilAssets.Num() + AssetClasses.Num() > 0)
 			{
 				for (AActor* Actor : SelectedActors)
-				{			
-					if(Actor)
+				{
+					if (Actor)
 					{
-						for(FAssetData& UtilAsset : UtilAssets)
+						auto ProcessAssetAction = [&SupportedUtils, &SupportedActors, &UtilityAndSelectionIndices, Actor](UActorActionUtility* InAction)
+						{
+							bool bPassesClassFilter = false;
+
+							UClass* SupportedClass = InAction->GetSupportedClass();
+							if (SupportedClass == nullptr || (SupportedClass && Actor->GetClass()->IsChildOf(SupportedClass)))
+							{
+								bPassesClassFilter = true;
+							}
+							
+							if (bPassesClassFilter)
+							{
+								SupportedUtils.AddUnique(InAction);
+								const int32 Index = SupportedActors.AddUnique(Actor);
+								UtilityAndSelectionIndices.FindOrAdd(InAction).Add(Index);
+							}
+						};
+						
+						// Process asset based utilities
+						for (const FAssetData& UtilAsset : UtilAssets)
 						{
-							if(UEditorUtilityBlueprint* Blueprint = Cast<UEditorUtilityBlueprint>(UtilAsset.GetAsset()))
+							if (UEditorUtilityBlueprint* Blueprint = Cast<UEditorUtilityBlueprint>(UtilAsset.GetAsset()))
 							{
-								if(UClass* BPClass = Blueprint->GeneratedClass.Get())
+								if (UClass* BPClass = Blueprint->GeneratedClass.Get())
 								{
-									if(UActorActionUtility* DefaultObject = Cast<UActorActionUtility>(BPClass->GetDefaultObject()))
+									if (UActorActionUtility* DefaultObject = Cast<UActorActionUtility>(BPClass->GetDefaultObject()))
 									{
-										UClass* SupportedClass = DefaultObject->GetSupportedClass();
-										if(SupportedClass == nullptr || (SupportedClass && Actor->GetClass()->IsChildOf(SupportedClass)))
-										{
-											SupportedUtils.AddUnique(DefaultObject);
-										}
+										ProcessAssetAction(DefaultObject);
 									}
 								}
 							}
 						}
+
+						// Process non-asset based utilities
+						for (UActorActionUtility* Action : AssetClasses)
+						{
+							ProcessAssetAction(Action);
+						}
 					}
 				}
 			}
@@ -68,7 +103,7 @@ public:
 				"ActorControl",
 				EExtensionHook::After,
 				CommandList,
-				FMenuExtensionDelegate::CreateStatic(&FBlutilityMenuExtensions::CreateBlutilityActionsMenu, SupportedUtils));
+				FMenuExtensionDelegate::CreateStatic(&FBlutilityMenuExtensions::CreateActorBlutilityActionsMenu, UtilityAndSelectionIndices, SupportedActors));
 		}
 
 		return Extender;
diff --git a/Engine/Source/Editor/Blutility/Private/BlutilityMenuExtensions.cpp b/Engine/Source/Editor/Blutility/Private/BlutilityMenuExtensions.cpp
--- a/Engine/Source/Editor/Blutility/Private/BlutilityMenuExtensions.cpp
+++ b/Engine/Source/Editor/Blutility/Private/BlutilityMenuExtensions.cpp
@@ -21,6 +21,10 @@
 #include "Widgets/Layout/SScrollBox.h"
 #include "ScopedTransaction.h"
 #include "Subsystems/AssetEditorSubsystem.h"
+#include "AssetData.h"
+#include "PropertyPathHelpers.h"
+
+#include "UObject/NoExportTypes.h"
 
 #define LOCTEXT_NAMESPACE "BlutilityMenuExtensions"
 
@@ -37,7 +41,7 @@ class SFunctionParamDialog : public SCompoundWidget
 
 	SLATE_END_ARGS()
 
-	void Construct(const FArguments& InArgs, TWeakPtr<SWindow> InParentWindow, TSharedRef<FStructOnScope> InStructOnScope)
+	void Construct(const FArguments& InArgs, TWeakPtr<SWindow> InParentWindow, TSharedRef<FStructOnScope> InStructOnScope, FName HiddenPropertyName)
 	{
 		bOKPressed = false;
 
@@ -67,9 +71,10 @@ class SFunctionParamDialog : public SCompoundWidget
 		FPropertyEditorModule& PropertyEditorModule = FModuleManager::Get().LoadModuleChecked<FPropertyEditorModule>("PropertyEditor");
 		TSharedRef<IStructureDetailsView> StructureDetailsView = PropertyEditorModule.CreateStructureDetailView(DetailsViewArgs, StructureViewArgs, InStructOnScope);
 
-		StructureDetailsView->GetDetailsView()->SetIsPropertyVisibleDelegate(FIsPropertyVisible::CreateLambda([](const FPropertyAndParent& InPropertyAndParent)
+		// Hide any property that has been marked as such
+		StructureDetailsView->GetDetailsView()->SetIsPropertyVisibleDelegate(FIsPropertyVisible::CreateLambda([HiddenPropertyName](const FPropertyAndParent& InPropertyAndParent)
 		{
-			return InPropertyAndParent.Property.HasAnyPropertyFlags(CPF_Parm);
+			return InPropertyAndParent.Property.HasAnyPropertyFlags(CPF_Parm) && InPropertyAndParent.Property.GetFName() != HiddenPropertyName;
 		}));
 
 		StructureDetailsView->GetDetailsView()->ForceRefresh();
@@ -186,32 +191,75 @@ void FBlutilityMenuExtensions::GetBlutilityClasses(TArray<FAssetData>& OutAssets
 		}
 	}
 }
+void FBlutilityMenuExtensions::CreateActorBlutilityActionsMenu(FMenuBuilder& MenuBuilder, TMap<class IEditorUtilityExtension*, TSet<int32>> Utils, const TArray<AActor*> SelectedSupportedActors)
+{
+	CreateBlutilityActionsMenu<AActor*>(MenuBuilder, Utils,
+		LOCTEXT("ScriptedActorActions", "Scripted Actor Actions"),
+		LOCTEXT("ScriptedActorActionsTooltip", "Scripted actions available for the selected actors"),
+		[](const FProperty* Property) -> bool
+		{
+			const FFieldClass* ClassOfProperty = Property->GetClass();
+			if (ClassOfProperty == FObjectProperty::StaticClass())
+			{
+				const FObjectProperty* ObjectProperty = CastField<const FObjectProperty>(Property);
+				return ObjectProperty->PropertyClass == AActor::StaticClass();
+			}
 
-void FBlutilityMenuExtensions::CreateBlutilityActionsMenu(FMenuBuilder& MenuBuilder, TArray<IEditorUtilityExtension*> Utils)
+			return false;
+		},
+		SelectedSupportedActors
+	);
+}
+
+void FBlutilityMenuExtensions::CreateAssetBlutilityActionsMenu(FMenuBuilder& MenuBuilder, TMap<class IEditorUtilityExtension*, TSet<int32>> Utils, const TArray<FAssetData> SelectedSupportedAssets)
 {
-	const static FName NAME_CallInEditor(TEXT("CallInEditor"));
+	CreateBlutilityActionsMenu<FAssetData>(MenuBuilder, Utils,
+		LOCTEXT("ScriptedAssetActions", "Scripted Asset Actions"),
+		LOCTEXT("ScriptedAssetActionsTooltip", "Scripted actions available for the selected assets"),
+		[](const FProperty* Property) -> bool
+		{
+			const FFieldClass* ClassOfProperty = Property->GetClass();
+			if (ClassOfProperty == FStructProperty::StaticClass())
+			{
+				const FStructProperty* StructProperty = CastField<const FStructProperty>(Property);
+				return StructProperty->Struct->GetName() == TEXT("AssetData");
+			}
 
-	// Helper struct to track the util to call a function on
-	struct FFunctionAndUtil
-	{
-		FFunctionAndUtil(UFunction* InFunction, IEditorUtilityExtension* InUtil)
-			: Function(InFunction)
-			, Util(InUtil) {}
+			return false;
+		},
+		SelectedSupportedAssets
+	);	
+}
 
-		bool operator==(const FFunctionAndUtil& InFunction) const
+void FBlutilityMenuExtensions::OpenEditorForUtility(const FFunctionAndUtil& FunctionAndUtil)
+{
+	// Edit the script if we have shift held down
+	if (UBlueprint* Blueprint = Cast<UBlueprint>(Cast<UObject>(FunctionAndUtil.Util)->GetClass()->ClassGeneratedBy))
+	{
+		if (IAssetEditorInstance* AssetEditor = GEditor->GetEditorSubsystem<UAssetEditorSubsystem>()->FindEditorForAsset(Blueprint, true))
+		{
+			check(AssetEditor->GetEditorName() == TEXT("BlueprintEditor"));
+			IBlueprintEditor* BlueprintEditor = static_cast<IBlueprintEditor*>(AssetEditor);
+			BlueprintEditor->JumpToHyperlink(FunctionAndUtil.Function, false);
+		}
+		else
 		{
-			return InFunction.Function == Function;
+			FBlueprintEditorModule& BlueprintEditorModule = FModuleManager::LoadModuleChecked<FBlueprintEditorModule>("Kismet");
+			TSharedRef<IBlueprintEditor> BlueprintEditor = BlueprintEditorModule.CreateBlueprintEditor(EToolkitMode::Standalone, TSharedPtr<IToolkitHost>(), Blueprint, false);
+			BlueprintEditor->JumpToHyperlink(FunctionAndUtil.Function, false);
 		}
+	}
+}
 
-		UFunction* Function;
-		IEditorUtilityExtension* Util;
-	};
-	TArray<FFunctionAndUtil> FunctionsToList;
+void FBlutilityMenuExtensions::ExtractFunctions(TMap<class IEditorUtilityExtension*, TSet<int32>>& Utils, TMap<FString, TArray<FFunctionAndUtil>>& OutCategoryFunctions)
+{
 	TSet<UClass*> ProcessedClasses;
+	const static FName NAME_CallInEditor(TEXT("CallInEditor"));
 
 	// Find the exposed functions available in each class, making sure to not list shared functions from a parent class more than once
-	for(IEditorUtilityExtension* Util : Utils)
+	for (TPair< IEditorUtilityExtension*, TSet<int32>> UtilitySelectionPair : Utils)
 	{
+		IEditorUtilityExtension* Util = UtilitySelectionPair.Key;
 		UClass* Class = Cast<UObject>(Util)->GetClass();
 
 		if (ProcessedClasses.Contains(Class))
@@ -230,110 +278,177 @@ void FBlutilityMenuExtensions::CreateBlutilityActionsMenu(FMenuBuilder& MenuBuil
 			{
 				if (Func->HasMetaData(NAME_CallInEditor) && Func->GetReturnProperty() == nullptr)
 				{
-					FunctionsToList.AddUnique(FFunctionAndUtil(Func, Util));
+					const static FName NAME_Category(TEXT("Category"));
+
+					const FString& FunctionCategory = Func->GetMetaData(NAME_Category);
+					TArray<FFunctionAndUtil>& Functions = OutCategoryFunctions.FindOrAdd(FunctionCategory);
+
+					Functions.AddUnique(FFunctionAndUtil(Func, Util, UtilitySelectionPair.Value));
 				}
 			}
 		}
 	}
 
-	// Sort the functions by name
-	FunctionsToList.Sort([](const FFunctionAndUtil& A, const FFunctionAndUtil& B) { return A.Function->GetName() < B.Function->GetName(); });
+	for (TPair<FString, TArray<FFunctionAndUtil>>& CategoryFunctionPair : OutCategoryFunctions)
+	{
+		// Sort the functions by name
+		CategoryFunctionPair.Value.Sort([](const FFunctionAndUtil& A, const FFunctionAndUtil& B) { return A.Function->GetName() < B.Function->GetName(); });
+	}
+}
 
-	// Add a menu item for each function
-	if (FunctionsToList.Num() > 0)
+template<typename SelectionType>
+void FBlutilityMenuExtensions::CreateBlutilityActionsMenu(FMenuBuilder& MenuBuilder, TMap<IEditorUtilityExtension*, TSet<int32>> Utils, const FText& MenuLabel, const FText& MenuToolTip, TFunction<bool(const FProperty * Property)> IsValidPropertyType, const TArray<SelectionType> Selection)
+{
+	TMap<FString, TArray<FFunctionAndUtil>> CategoryFunctions;
+	ExtractFunctions(Utils, CategoryFunctions);
+
+	auto AddFunctionEntries = [Selection, IsValidPropertyType](FMenuBuilder& SubMenuBuilder, const TArray<FFunctionAndUtil>& FunctionUtils)
 	{
-		MenuBuilder.AddSubMenu(
-			LOCTEXT("ScriptedActorActions", "Scripted Actions"), 
-			LOCTEXT("ScriptedActorActionsTooltip", "Scripted actions available for the selected actors"),
-			FNewMenuDelegate::CreateLambda([FunctionsToList](FMenuBuilder& InMenuBuilder)
-			{
-				for (const FFunctionAndUtil& FunctionAndUtil : FunctionsToList)
-				{
-					const FText TooltipText = FText::Format(LOCTEXT("AssetUtilTooltipFormat", "{0}\n(Shift-click to edit script)"), FunctionAndUtil.Function->GetToolTipText());
+		for (const FFunctionAndUtil& FunctionAndUtil : FunctionUtils)
+		{
+			const FText TooltipText = FText::Format(LOCTEXT("AssetUtilTooltipFormat", "{0}\n(Shift-click to edit script)"), FunctionAndUtil.Function->GetToolTipText());
 
-					InMenuBuilder.AddMenuEntry(
-						FunctionAndUtil.Function->GetDisplayNameText(), 
-						TooltipText,
-						FSlateIcon("EditorStyle", "GraphEditor.Event_16x"),
-						FExecuteAction::CreateLambda([FunctionAndUtil] 
+			SubMenuBuilder.AddMenuEntry(
+				FunctionAndUtil.Function->GetDisplayNameText(),
+				TooltipText,
+				FSlateIcon("EditorStyle", "GraphEditor.Event_16x"),
+				FExecuteAction::CreateLambda([FunctionAndUtil, Selection, IsValidPropertyType]
+				{
+					if (FSlateApplication::Get().GetModifierKeys().IsShiftDown())
+					{
+						OpenEditorForUtility(FunctionAndUtil);
+					}
+					else
+					{
+						// We dont run this on the CDO, as bad things could occur!
+						UObject* TempObject = NewObject<UObject>(GetTransientPackage(), Cast<UObject>(FunctionAndUtil.Util)->GetClass());
+						TempObject->AddToRoot(); // Some Blutility actions might run GC so the TempObject needs to be rooted to avoid getting destroyed
+
+						if (FunctionAndUtil.Function->NumParms > 0)
 						{
-							if(FSlateApplication::Get().GetModifierKeys().IsShiftDown())
+							// Create a parameter struct and fill in defaults
+							TSharedRef<FStructOnScope> FuncParams = MakeShared<FStructOnScope>(FunctionAndUtil.Function);
+
+							FProperty* FirstParamProperty = nullptr;
+
+							int32 ParameterIndex = 0;
+							for (TFieldIterator<FProperty> It(FunctionAndUtil.Function); It&& It->HasAnyPropertyFlags(CPF_Parm); ++It)
 							{
-								// Edit the script if we have shift held down
-								if(UBlueprint* Blueprint = Cast<UBlueprint>(Cast<UObject>(FunctionAndUtil.Util)->GetClass()->ClassGeneratedBy))
+								FString Defaults;
+								if (UEdGraphSchema_K2::FindFunctionParameterDefaultValue(FunctionAndUtil.Function, *It, Defaults))
 								{
-									if(IAssetEditorInstance* AssetEditor = GEditor->GetEditorSubsystem<UAssetEditorSubsystem>()->FindEditorForAsset(Blueprint, true))
-									{
-										check(AssetEditor->GetEditorName() == TEXT("BlueprintEditor"));
-										IBlueprintEditor* BlueprintEditor = static_cast<IBlueprintEditor*>(AssetEditor);
-										BlueprintEditor->JumpToHyperlink(FunctionAndUtil.Function, false);
-									}
-									else
-									{
-										FBlueprintEditorModule& BlueprintEditorModule = FModuleManager::LoadModuleChecked<FBlueprintEditorModule>("Kismet");
-										TSharedRef<IBlueprintEditor> BlueprintEditor = BlueprintEditorModule.CreateBlueprintEditor(EToolkitMode::Standalone, TSharedPtr<IToolkitHost>(), Blueprint, false);
-										BlueprintEditor->JumpToHyperlink(FunctionAndUtil.Function, false);
-									}
+									It->ImportText(*Defaults, It->ContainerPtrToValuePtr<uint8>(FuncParams->GetStructMemory()), PPF_None, nullptr);
 								}
-							}
-							else
-							{
-								// We dont run this on the CDO, as bad things could occur!
-								UObject* TempObject = NewObject<UObject>(GetTransientPackage(), Cast<UObject>(FunctionAndUtil.Util)->GetClass());
-								TempObject->AddToRoot(); // Some Blutility actions might run GC so the TempObject needs to be rooted to avoid getting destroyed
 
-								if(FunctionAndUtil.Function->NumParms > 0)
+								// Check to see if the first parameter matches the selection object type, in that case we can directly forward the selection to it
+								if (ParameterIndex == 0 && IsValidPropertyType(*It))
 								{
-									// Create a parameter struct and fill in defaults
-									TSharedRef<FStructOnScope> FuncParams = MakeShared<FStructOnScope>(FunctionAndUtil.Function);
-									for (TFieldIterator<FProperty> It(FunctionAndUtil.Function); It && It->HasAnyPropertyFlags(CPF_Parm); ++It)
-									{
-										FString Defaults;
-										if(UEdGraphSchema_K2::FindFunctionParameterDefaultValue(FunctionAndUtil.Function, *It, Defaults))
-										{
-											It->ImportText(*Defaults, It->ContainerPtrToValuePtr<uint8>(FuncParams->GetStructMemory()), PPF_None, nullptr);
-										}
-									}
+									FirstParamProperty = *It;
+								}
+
+								++ParameterIndex;
+							}
 
-									// pop up a dialog to input params to the function
-									TSharedRef<SWindow> Window = SNew(SWindow)
-										.Title(FunctionAndUtil.Function->GetDisplayNameText())
-										.ClientSize(FVector2D(400, 200))
-										.SupportsMinimize(false)
-										.SupportsMaximize(false);
+							bool bApply = true;
 
-									TSharedPtr<SFunctionParamDialog> Dialog;
-									Window->SetContent(
-										SAssignNew(Dialog, SFunctionParamDialog, Window, FuncParams)
-										.OkButtonText(LOCTEXT("OKButton", "OK"))
-										.OkButtonTooltipText(FunctionAndUtil.Function->GetToolTipText()));
+							if (!FirstParamProperty || ParameterIndex > 1)
+							{
+								// pop up a dialog to input params to the function
+								TSharedRef<SWindow> Window = SNew(SWindow)
+									.Title(FunctionAndUtil.Function->GetDisplayNameText())
+									.ClientSize(FVector2D(400, 200))
+									.SupportsMinimize(false)
+									.SupportsMaximize(false);
+
+								TSharedPtr<SFunctionParamDialog> Dialog;
+								Window->SetContent(
+									SAssignNew(Dialog, SFunctionParamDialog, Window, FuncParams, FirstParamProperty ? FirstParamProperty->GetFName() : NAME_None)
+									.OkButtonText(LOCTEXT("OKButton", "OK"))
+									.OkButtonTooltipText(FunctionAndUtil.Function->GetToolTipText()));
+
+								GEditor->EditorAddModalWindow(Window);
+								bApply = Dialog->bOKPressed;
+							}
 
-									GEditor->EditorAddModalWindow(Window);
 
-									if(Dialog->bOKPressed)
+							if (bApply)
+							{
+								FScopedTransaction Transaction(NSLOCTEXT("UnrealEd", "BlutilityAction", "Blutility Action"));
+								FEditorScriptExecutionGuard ScriptGuard;
+								const bool bForwardUserSelection = FirstParamProperty != nullptr;
+								if (bForwardUserSelection)
+								{
+									// For each user-select asset forward the selection object into the function first's parameter (if it matches)
+									const FString Path = FirstParamProperty->GetPathName(FunctionAndUtil.Function);
+
+									// Ensure we only process selection objects that are valid for this function/utility
+									for (const int32& SelectionIndex : FunctionAndUtil.SelectionIndices)
 									{
-										FScopedTransaction Transaction( NSLOCTEXT("UnrealEd", "BlutilityAction", "Blutility Action") );
-										FEditorScriptExecutionGuard ScriptGuard;
+										const auto SelectedAsset = Selection[SelectionIndex];
+										FirstParamProperty->CopySingleValue(FirstParamProperty->ContainerPtrToValuePtr<uint8>(FuncParams->GetStructMemory()), &SelectedAsset);
 										TempObject->ProcessEvent(FunctionAndUtil.Function, FuncParams->GetStructMemory());
 									}
 								}
 								else
 								{
-									FScopedTransaction Transaction( NSLOCTEXT("UnrealEd", "BlutilityAction", "Blutility Action") );
-									FEditorScriptExecutionGuard ScriptGuard;
-									TempObject->ProcessEvent(FunctionAndUtil.Function, nullptr);
+									// User is expected to manage the asset selection on its own
+									TempObject->ProcessEvent(FunctionAndUtil.Function, FuncParams->GetStructMemory());
 								}
-
-								TempObject->RemoveFromRoot();
 							}
-						}));
+						}
+						else
+						{
+							FScopedTransaction Transaction(NSLOCTEXT("UnrealEd", "BlutilityAction", "Blutility Action"));
+							FEditorScriptExecutionGuard ScriptGuard;
+							TempObject->ProcessEvent(FunctionAndUtil.Function, nullptr);
+						}
+
+						TempObject->RemoveFromRoot();
+					}
+				})
+			);
+		}
+	};
+
+	// Add a menu item for each function
+	if (CategoryFunctions.Num() > 0)
+	{
+		MenuBuilder.AddSubMenu(
+			MenuLabel,
+			MenuToolTip,
+			FNewMenuDelegate::CreateLambda([CategoryFunctions, AddFunctionEntries](FMenuBuilder& InMenuBuilder)
+			{
+				TArray<FString> CategoryNames;
+				CategoryFunctions.GenerateKeyArray(CategoryNames);
+				CategoryNames.Remove(FString());
+				CategoryNames.Sort();
+				
+				// Add functions belong to the same category to a sub-menu
+				for (const FString& CategoryName : CategoryNames)
+				{
+					const TArray<FFunctionAndUtil>& FunctionUtils = CategoryFunctions.FindChecked(CategoryName);
+					InMenuBuilder.AddSubMenu(FText::FromString(CategoryName), FText::FromString(CategoryName),
+						FNewMenuDelegate::CreateLambda([FunctionUtils, AddFunctionEntries](FMenuBuilder& InSubMenuBuilder) 
+						{
+							AddFunctionEntries(InSubMenuBuilder, FunctionUtils);
+						})
+					);
+				}
+
+				// Non-categorized functions
+				const TArray<FFunctionAndUtil>* DefaultCategoryFunctionsPtr = CategoryFunctions.Find(FString());
+				if (DefaultCategoryFunctionsPtr)
+				{
+					AddFunctionEntries(InMenuBuilder, *DefaultCategoryFunctionsPtr);
 				}
 			}),
 			false,
-			FSlateIcon("EditorStyle", "GraphEditor.Event_16x"));
+			FSlateIcon("EditorStyle", "GraphEditor.Event_16x")
+		);
 	}
 }
 
+
 UEditorUtilityExtension::UEditorUtilityExtension(const class FObjectInitializer& ObjectInitializer)
 	: Super(ObjectInitializer)
 {
diff --git a/Engine/Source/Editor/Blutility/Private/BlutilityMenuExtensions.h b/Engine/Source/Editor/Blutility/Private/BlutilityMenuExtensions.h
--- a/Engine/Source/Editor/Blutility/Private/BlutilityMenuExtensions.h
+++ b/Engine/Source/Editor/Blutility/Private/BlutilityMenuExtensions.h
@@ -18,7 +18,36 @@ public:
 	static void GetBlutilityClasses(TArray<FAssetData>& OutAssets, const FName& InClassName);
 
 	/** Helper function that populates a menu based on the exposed functions in a set of Blutility objects */
-	static void CreateBlutilityActionsMenu(FMenuBuilder& MenuBuilder, TArray<class IEditorUtilityExtension*> Utils);
+	static void CreateAssetBlutilityActionsMenu(FMenuBuilder& MenuBuilder, TMap<class IEditorUtilityExtension*, TSet<int32>> Utils, const TArray<FAssetData> SelectedSupportedAssets);
+	static void CreateActorBlutilityActionsMenu(FMenuBuilder& MenuBuilder, TMap<class IEditorUtilityExtension*, TSet<int32>> Utils, const TArray<AActor*> SelectedSupportedActors);
+	
+protected:
+	// Helper struct to track the util to call a function on
+	struct FFunctionAndUtil
+	{
+		FFunctionAndUtil(UFunction* InFunction, IEditorUtilityExtension* InUtil, TSet<int32>& InSelection)
+			: Function(InFunction)
+			, Util(InUtil) 
+			, SelectionIndices(InSelection)
+		{}
+
+		bool operator==(const FFunctionAndUtil& InFunction) const
+		{
+			return InFunction.Function == Function && InFunction.Util == Util;
+		}
+
+		UFunction* Function;
+		IEditorUtilityExtension* Util;
+		/** Indices to original object selection array which are supported by this utility */
+		TSet<int32> SelectionIndices;
+	};
+
+protected:
+	template<typename SelectionType>
+	static void CreateBlutilityActionsMenu(FMenuBuilder& MenuBuilder, TMap<class IEditorUtilityExtension*, TSet<int32>> Utils, const FText& MenuLabel, const FText& MenuToolTip, TFunction<bool(const FProperty * Property)> IsValidPropertyType, const TArray<SelectionType> Selection);
+
+	static void OpenEditorForUtility(const FFunctionAndUtil& FunctionAndUtil);
+	static void ExtractFunctions(TMap<class IEditorUtilityExtension*, TSet<int32>>& Utils, TMap<FString, TArray<FFunctionAndUtil>>& OutCategoryFunctions);
 };
 
 UINTERFACE(BlueprintType)
@@ -30,4 +59,4 @@ class UEditorUtilityExtension : public UInterface
 class IEditorUtilityExtension
 {
 	GENERATED_IINTERFACE_BODY()
-};
+};
\ No newline at end of file
