Added function to get the current Content Browser path in editor blueprints and python

diff --git a/Engine/Source/Editor/Blutility/Private/EditorUtilityLibrary.cpp b/Engine/Source/Editor/Blutility/Private/EditorUtilityLibrary.cpp
--- a/Engine/Source/Editor/Blutility/Private/EditorUtilityLibrary.cpp
+++ b/Engine/Source/Editor/Blutility/Private/EditorUtilityLibrary.cpp
@@ -301,6 +301,13 @@ AActor* UEditorUtilityLibrary::GetActorReference(FString PathToActor)
 #endif //WITH_EDITOR
 }
 
+bool UEditorUtilityLibrary::GetCurrentContentBrowserPath(FString& OutPath)
+{
+	IContentBrowserSingleton& ContentBrowser = FModuleManager::LoadModuleChecked<FContentBrowserModule>("ContentBrowser").Get();
+	OutPath = ContentBrowser.GetCurrentPath(EContentBrowserPathType::Internal);
+	return !OutPath.IsEmpty();
+}
+
 #endif
 
 #undef LOCTEXT_NAMESPACE
diff --git a/Engine/Source/Editor/Blutility/Public/EditorUtilityLibrary.h b/Engine/Source/Editor/Blutility/Public/EditorUtilityLibrary.h
--- a/Engine/Source/Editor/Blutility/Public/EditorUtilityLibrary.h
+++ b/Engine/Source/Editor/Blutility/Public/EditorUtilityLibrary.h
@@ -158,5 +158,12 @@ public:
 	UFUNCTION(BlueprintPure, Category = "Development|Editor")
 	AActor* GetActorReference(FString PathToActor);
 
+	/**
+	 * Attempts to get the path for the active content browser, returns false if there is no active content browser
+	 * @param	OutPath	The returned path if successfully found
+	 * @return	Whether a path was successfully returned
+	 */
+	UFUNCTION(BlueprintPure, Category = "Development|Editor")
+	static bool GetCurrentContentBrowserPath(FString& OutPath);
 #endif
 };
