PR #8391: Added UAssetActionUtility to Picker Common Classes (Contributed by [REDACTED_PERSON])

diff --git a/Engine/Source/Editor/Blutility/Private/EditorUtilityBlueprintFactory.cpp b/Engine/Source/Editor/Blutility/Private/EditorUtilityBlueprintFactory.cpp
--- a/Engine/Source/Editor/Blutility/Private/EditorUtilityBlueprintFactory.cpp
+++ b/Engine/Source/Editor/Blutility/Private/EditorUtilityBlueprintFactory.cpp
@@ -17,6 +17,7 @@
 #include "EditorUtilityWidget.h"
 #include "EditorUtilityActor.h"
 #include "EditorUtilityObject.h"
+#include "AssetActionUtility.h"
 
 class FBlutilityBlueprintFactoryFilter : public IClassViewerFilter
 {
@@ -78,6 +79,7 @@ bool UEditorUtilityBlueprintFactory::ConfigureProperties()
 
 	Options.ExtraPickerCommonClasses.Add(AEditorUtilityActor::StaticClass());
 	Options.ExtraPickerCommonClasses.Add(UEditorUtilityObject::StaticClass());
+	Options.ExtraPickerCommonClasses.Add(UAssetActionUtility::StaticClass());
 
 	TSharedPtr< FBlutilityBlueprintFactoryFilter > Filter = MakeShareable(new FBlutilityBlueprintFactoryFilter);
 	Filter->DisallowedChildOfClasses.Add(UEditorUtilityWidget::StaticClass());
