Fix a crash when we force delete a Editor Utility Widget when it is running

diff --git a/Engine/Source/Editor/Blutility/Private/EditorUtilityWidgetBlueprint.cpp b/Engine/Source/Editor/Blutility/Private/EditorUtilityWidgetBlueprint.cpp
--- a/Engine/Source/Editor/Blutility/Private/EditorUtilityWidgetBlueprint.cpp
+++ b/Engine/Source/Editor/Blutility/Private/EditorUtilityWidgetBlueprint.cpp
@@ -129,11 +129,14 @@ void UEditorUtilityWidgetBlueprint::ChangeTabWorld(UWorld* World, EMapChangeType
 
 void UEditorUtilityWidgetBlueprint::UpdateRespawnListIfNeeded(TSharedRef<SDockTab> TabBeingClosed)
 {
-	const UEditorUtilityWidget* EditorUtilityWidget = GeneratedClass->GetDefaultObject<UEditorUtilityWidget>();
-	if (EditorUtilityWidget && EditorUtilityWidget->ShouldAlwaysReregisterWithWindowsMenu() == false)
+	if (GeneratedClass)
 	{
-		IBlutilityModule* BlutilityModule = FModuleManager::GetModulePtr<IBlutilityModule>("Blutility");
-		BlutilityModule->RemoveLoadedScriptUI(this);
+		const UEditorUtilityWidget* EditorUtilityWidget = GeneratedClass->GetDefaultObject<UEditorUtilityWidget>();
+		if (EditorUtilityWidget && EditorUtilityWidget->ShouldAlwaysReregisterWithWindowsMenu() == false)
+		{
+			IBlutilityModule* BlutilityModule = FModuleManager::GetModulePtr<IBlutilityModule>("Blutility");
+			BlutilityModule->RemoveLoadedScriptUI(this);
+		}
 	}
 	CreatedUMGWidget = nullptr;
 }
