Added function to get the currently selected folder(s) path in utility widgets and Python

diff --git a/Engine/Source/Editor/Blutility/Private/EditorUtilityLibrary.cpp b/Engine/Source/Editor/Blutility/Private/EditorUtilityLibrary.cpp
--- a/Engine/Source/Editor/Blutility/Private/EditorUtilityLibrary.cpp
+++ b/Engine/Source/Editor/Blutility/Private/EditorUtilityLibrary.cpp
@@ -316,6 +316,14 @@ bool UEditorUtilityLibrary::GetCurrentContentBrowserPath(FString& OutPath)
 	}
 }
 
+TArray<FString> UEditorUtilityLibrary::GetSelectedFolderPaths()
+{
+	IContentBrowserSingleton& ContentBrowser = FModuleManager::LoadModuleChecked<FContentBrowserModule>("ContentBrowser").Get();
+	TArray<FString> Paths;
+	ContentBrowser.GetSelectedFolders(Paths);
+	return Paths;
+}
+
 #endif
 
 #undef LOCTEXT_NAMESPACE
diff --git a/Engine/Source/Editor/Blutility/Public/EditorUtilityLibrary.h b/Engine/Source/Editor/Blutility/Public/EditorUtilityLibrary.h
--- a/Engine/Source/Editor/Blutility/Public/EditorUtilityLibrary.h
+++ b/Engine/Source/Editor/Blutility/Public/EditorUtilityLibrary.h
@@ -165,5 +165,9 @@ public:
 	 */
 	UFUNCTION(BlueprintPure, Category = "Development|Editor")
 	static bool GetCurrentContentBrowserPath(FString& OutPath);
+
+	// Gets the path to the currently selected folder in the content browser
+	UFUNCTION(BlueprintPure, Category = "Development|Editor")
+	static TArray<FString> GetSelectedFolderPaths();
 #endif
 };
