Add PIE callbacks to editor utility subsystem.

diff --git a/Engine/Source/Editor/Blutility/Private/EditorUtilitySubsystem.cpp b/Engine/Source/Editor/Blutility/Private/EditorUtilitySubsystem.cpp
--- a/Engine/Source/Editor/Blutility/Private/EditorUtilitySubsystem.cpp
+++ b/Engine/Source/Editor/Blutility/Private/EditorUtilitySubsystem.cpp
@@ -50,6 +50,9 @@ void UEditorUtilitySubsystem::Initialize(FSubsystemCollectionBase& Collection)
 	}
 
 	TickerHandle = FTSTicker::GetCoreTicker().AddTicker(FTickerDelegate::CreateUObject(this, &UEditorUtilitySubsystem::Tick), 0);
+
+	FEditorDelegates::BeginPIE.AddUObject(this, &UEditorUtilitySubsystem::HandleOnBeginPIE);
+	FEditorDelegates::EndPIE.AddUObject(this, &UEditorUtilitySubsystem::HandleOnEndPIE);
 }
 
 void UEditorUtilitySubsystem::Deinitialize()
@@ -62,6 +65,9 @@ void UEditorUtilitySubsystem::Deinitialize()
 	FTSTicker::GetCoreTicker().RemoveTicker(TickerHandle);
 
 	IConsoleManager::Get().UnregisterConsoleObject(RunTaskCommandObject);
+
+	FEditorDelegates::BeginPIE.RemoveAll(this);
+	FEditorDelegates::EndPIE.RemoveAll(this);
 }
 
 void UEditorUtilitySubsystem::AddReferencedObjects(UObject* InThis, FReferenceCollector& Collector)
@@ -483,4 +489,14 @@ UClass* UEditorUtilitySubsystem::FindBlueprintClass(const FString& TargetNameRaw
 	return FoundClass;
 }
 
+void UEditorUtilitySubsystem::HandleOnBeginPIE(const bool bIsSimulating)
+{
+	OnBeginPIE.Broadcast(bIsSimulating);
+}
+
+void UEditorUtilitySubsystem::HandleOnEndPIE(const bool bIsSimulating)
+{
+	OnEndPIE.Broadcast(bIsSimulating);
+}
+
 #undef LOCTEXT_NAMESPACE
diff --git a/Engine/Source/Editor/Blutility/Public/EditorUtilitySubsystem.h b/Engine/Source/Editor/Blutility/Public/EditorUtilitySubsystem.h
--- a/Engine/Source/Editor/Blutility/Public/EditorUtilitySubsystem.h
+++ b/Engine/Source/Editor/Blutility/Public/EditorUtilitySubsystem.h
@@ -16,6 +16,9 @@ class SWindow;
 class UEditorUtilityWidget;
 class UEditorUtilityTask;
 
+/** Delegate for a PIE event exposed via Editor Utility (begin, end, pause/resume, etc) */
+DECLARE_DYNAMIC_MULTICAST_DELEGATE_OneParam(FOnEditorUtilityPIEEvent, const bool, bIsSimulating);
+
 UCLASS(config = EditorPerProjectUserSettings)
 class BLUTILITY_API UEditorUtilitySubsystem : public UEditorSubsystem
 {
@@ -84,6 +87,14 @@ public:
 	void RegisterReferencedObject(UObject* ObjectToReference);
 	void UnregisterReferencedObject(UObject* ObjectToReference);
 
+	/** Expose Begin PIE to blueprints.*/
+	UPROPERTY(BlueprintAssignable)
+	FOnEditorUtilityPIEEvent OnBeginPIE;
+
+	/** Expose End PIE to blueprints.*/
+	UPROPERTY(BlueprintAssignable)
+	FOnEditorUtilityPIEEvent OnEndPIE;
+
 protected:
 	UEditorUtilityTask* GetActiveTask() { return ActiveTaskStack.Num() > 0 ? ActiveTaskStack[ActiveTaskStack.Num() - 1] : nullptr; };
 
@@ -99,6 +110,12 @@ protected:
 	UClass* FindClassByName(const FString& RawTargetName);
 	UClass* FindBlueprintClass(const FString& TargetNameRaw);
 
+	/** Called when Play in Editor begins. */
+	void HandleOnBeginPIE(const bool bIsSimulating);
+
+	/** Called when Play in Editor stops. */
+	void HandleOnEndPIE(const bool bIsSimulating);
+
 private:
 	IConsoleObject* RunTaskCommandObject = nullptr;
 	IConsoleObject* CancelAllTasksCommandObject = nullptr;
