Add support for warm up frames when doing an AsyncCaptureScene.

diff --git a/Engine/Source/Editor/Blutility/Private/AsyncCaptureScene.cpp b/Engine/Source/Editor/Blutility/Private/AsyncCaptureScene.cpp
--- a/Engine/Source/Editor/Blutility/Private/AsyncCaptureScene.cpp
+++ b/Engine/Source/Editor/Blutility/Private/AsyncCaptureScene.cpp
@@ -30,12 +30,20 @@ UAsyncCaptureScene::UAsyncCaptureScene()
 UAsyncCaptureScene* UAsyncCaptureScene::CaptureSceneAsync(UCameraComponent* ViewCamera, TSubclassOf<ASceneCapture2D> SceneCaptureClass, int ResX, int ResY)
 {
 	UAsyncCaptureScene* AsyncTask = NewObject<UAsyncCaptureScene>();
-	AsyncTask->Start(ViewCamera, SceneCaptureClass, ResX, ResY);
+	AsyncTask->Start(ViewCamera, SceneCaptureClass, ResX, ResY, 0);
 
 	return AsyncTask;
 }
 
-void UAsyncCaptureScene::Start(UCameraComponent* ViewCamera, TSubclassOf<ASceneCapture2D> SceneCaptureClass, int ResX, int ResY)
+UAsyncCaptureScene* UAsyncCaptureScene::CaptureSceneWithWarmupAsync(UCameraComponent* ViewCamera, TSubclassOf<ASceneCapture2D> SceneCaptureClass, int ResX, int ResY, int WarmUpFrames)
+{
+	UAsyncCaptureScene* AsyncTask = NewObject<UAsyncCaptureScene>();
+	AsyncTask->Start(ViewCamera, SceneCaptureClass, ResX, ResY, WarmUpFrames);
+
+	return AsyncTask;
+}
+
+void UAsyncCaptureScene::Start(UCameraComponent* ViewCamera, TSubclassOf<ASceneCapture2D> SceneCaptureClass, int ResX, int ResY, int InWarmUpFrames)
 {
 	const FVector CaptureLocation = ViewCamera->GetComponentLocation();
 	const FRotator CaptureRotation = ViewCamera->GetComponentRotation();
@@ -64,6 +72,8 @@ void UAsyncCaptureScene::Start(UCameraComponent* ViewCamera, TSubclassOf<ASceneC
 		ViewCamera->GetCameraView(0, CaptureView);
 		CaptureComponent->SetCameraView(CaptureView);
 	}
+
+	WarmUpFrames = FMath::Max(InWarmUpFrames, 1);
 }
 
 void UAsyncCaptureScene::Activate()
@@ -80,7 +90,10 @@ void UAsyncCaptureScene::Activate()
 
 	FinishLoadingBeforeScreenshot();
 
-	CaptureComponent->CaptureScene();
+	for (int32 FrameCount = 0; FrameCount < WarmUpFrames; ++FrameCount)
+	{
+		CaptureComponent->CaptureScene();
+	}
 
 	NotifyComplete(SceneCaptureRT);
 }
diff --git a/Engine/Source/Editor/Blutility/Public/AsyncCaptureScene.h b/Engine/Source/Editor/Blutility/Public/AsyncCaptureScene.h
--- a/Engine/Source/Editor/Blutility/Public/AsyncCaptureScene.h
+++ b/Engine/Source/Editor/Blutility/Public/AsyncCaptureScene.h
@@ -28,6 +28,9 @@ public:
 
 	UFUNCTION(BlueprintCallable, meta=( BlueprintInternalUseOnly="true" ))
 	static UAsyncCaptureScene* CaptureSceneAsync(UCameraComponent* ViewCamera, TSubclassOf<ASceneCapture2D> SceneCaptureClass, int ResX, int ResY);
+	
+	UFUNCTION(BlueprintCallable, meta = (BlueprintInternalUseOnly = "true"))
+	static UAsyncCaptureScene* CaptureSceneWithWarmupAsync(UCameraComponent* ViewCamera, TSubclassOf<ASceneCapture2D> SceneCaptureClass, int ResX, int ResY, int WarmUpFrames);
 
 	virtual void Activate() override;
 public:
@@ -38,7 +41,7 @@ public:
 private:
 
 	void NotifyComplete(UTextureRenderTarget2D* InTexture);
-	void Start(UCameraComponent* ViewCamera, TSubclassOf<ASceneCapture2D> SceneCaptureClass, int ResX, int ResY);
+	void Start(UCameraComponent* ViewCamera, TSubclassOf<ASceneCapture2D> SceneCaptureClass, int ResX, int ResY, int InWarmUpFrames);
 	void FinishLoadingBeforeScreenshot();
 	
 private:
@@ -47,4 +50,6 @@ private:
 
 	UPROPERTY()
 	TObjectPtr<UTextureRenderTarget2D> SceneCaptureRT;
+
+	int32 WarmUpFrames;
 };
