Speculative fix crash FindUtilityWidgetFromBlueprint with invalid asset.

diff --git a/Engine/Source/Editor/Blutility/Private/EditorUtilitySubsystem.cpp b/Engine/Source/Editor/Blutility/Private/EditorUtilitySubsystem.cpp
--- a/Engine/Source/Editor/Blutility/Private/EditorUtilitySubsystem.cpp
+++ b/Engine/Source/Editor/Blutility/Private/EditorUtilitySubsystem.cpp
@@ -258,6 +258,12 @@ bool UEditorUtilitySubsystem::CloseTabByID(FName NewTabID)
 
 UEditorUtilityWidget* UEditorUtilitySubsystem::FindUtilityWidgetFromBlueprint(class UEditorUtilityWidgetBlueprint* InBlueprint)
 {
+	if (!IsValid(InBlueprint))
+	{ 
+		UE_LOG(LogEditorUtilityBlueprint, Error, TEXT("Found Invalid Blueprint in FindUtilityWidgetFromBlueprint"));
+		return nullptr; 
+	}
+
 	return InBlueprint->GetCreatedWidget();
 }
 
