Add an input component to Editor Utility actor so that you can bind to enhanced input (and other runtime input delegates) in editor utilities outside of PIE. This will make it way easier for virtual production teams to bind input of physical camera equipment to actors in the level. They need these bindings available during the level editor and during PIE, which is why Enhanced Input is being used as it already has the good support for what they want, we just need to expose it in the editor.

Listening to input is opt in per utility, and would currently only be triggered if you create an associated Input Action that is available to the editor.

diff --git a/Engine/Source/Editor/Blutility/Classes/EditorUtilityActor.h b/Engine/Source/Editor/Blutility/Classes/EditorUtilityActor.h
--- a/Engine/Source/Editor/Blutility/Classes/EditorUtilityActor.h
+++ b/Engine/Source/Editor/Blutility/Classes/EditorUtilityActor.h
@@ -11,6 +11,7 @@
 #include "UObject/ScriptMacros.h"
 #include "EditorUtilityActor.generated.h"
 
+class UInputComponent;
 
 UCLASS(Abstract, Blueprintable, meta = (ShowWorldContextPin))
 class BLUTILITY_API AEditorUtilityActor : public AActor
@@ -20,4 +21,40 @@ class BLUTILITY_API AEditorUtilityActor : public AActor
 	// Standard function to execute
 	UFUNCTION(BlueprintCallable, BlueprintImplementableEvent, Category = "Editor")
 	void Run();
+
+	virtual void OnConstruction(const FTransform& Transform) override;
+#if WITH_EDITOR
+	virtual void PostEditChangeProperty(FPropertyChangedEvent& PropertyChangedEvent) override;
+#endif
+
+	/** Returns the current InputComponent on this utility actor. This will be NULL unless bReceivesEditorInput is set to true. */
+	UFUNCTION(BlueprintCallable, BlueprintPure, Category = "Input|Editor")
+	UInputComponent* GetInputComponent() const
+	{ 
+		return EditorOnlyInputComponent.Get();
+	}
+	
+	UFUNCTION(BlueprintSetter, Category = "Input|Editor")
+	void SetReceivesEditorInput(bool bInValue);
+	
+	UFUNCTION(BlueprintGetter, BlueprintPure, Category = "Input|Editor")
+	bool GetReceivesEditorInput() const
+	{
+		return bReceivesEditorInput;
+	}	
+	
+private:
+
+	/** Creates the EditorOnlyInputComponent if it does not already exist and registers all subobject callbacks to it */
+	void CreateEditorInput();
+
+	/** Removes the EditorOnlyInputComponent from this utility actor */
+	void RemoveEditorInput();
+	
+	UPROPERTY(Transient, DuplicateTransient)
+	TObjectPtr<UInputComponent> EditorOnlyInputComponent;
+	
+	/** If set to true, then this actor will be able to recieve input delegate callbacks when in the editor. */
+	UPROPERTY(EditAnywhere, Category = "Input|Editor", BlueprintSetter = SetReceivesEditorInput, BlueprintGetter = GetReceivesEditorInput)
+	bool bReceivesEditorInput = false;
 };
diff --git a/Engine/Source/Editor/Blutility/Private/EditorUtilityActor.cpp b/Engine/Source/Editor/Blutility/Private/EditorUtilityActor.cpp
--- a/Engine/Source/Editor/Blutility/Private/EditorUtilityActor.cpp
+++ b/Engine/Source/Editor/Blutility/Private/EditorUtilityActor.cpp
@@ -1,7 +1,8 @@
 // Copyright Epic Games, Inc. All Rights Reserved.
 
 #include "EditorUtilityActor.h"
-
+#include "GameFramework/InputSettings.h"
+#include "Engine/InputDelegateBinding.h"
 
 /////////////////////////////////////////////////////
 
@@ -9,3 +10,57 @@ AEditorUtilityActor::AEditorUtilityActor(const FObjectInitializer& ObjectInitial
 	: Super(ObjectInitializer)
 {
 }
+
+void AEditorUtilityActor::OnConstruction(const FTransform& Transform)
+{
+	Super::OnConstruction(Transform);
+
+	CreateEditorInput();
+}
+
+void AEditorUtilityActor::CreateEditorInput()
+{
+	if (bReceivesEditorInput && !HasAnyFlags(RF_ClassDefaultObject) && !EditorOnlyInputComponent)
+	{
+		EditorOnlyInputComponent = NewObject<UInputComponent>(this, UInputSettings::GetDefaultInputComponentClass(), TEXT("EdUtilActor_InputComponent0"), RF_Transient);
+		UInputDelegateBinding::BindInputDelegatesWithSubojects(this, EditorOnlyInputComponent);
+	}
+}
+
+void AEditorUtilityActor::RemoveEditorInput()
+{
+	ensure(!bReceivesEditorInput);
+	
+	if (EditorOnlyInputComponent)
+	{
+		EditorOnlyInputComponent->DestroyComponent();
+	}
+	EditorOnlyInputComponent = nullptr;
+}
+
+#if WITH_EDITOR
+void AEditorUtilityActor::PostEditChangeProperty(FPropertyChangedEvent& PropertyChangedEvent)
+{
+	Super::PostEditChangeProperty(PropertyChangedEvent);
+
+	const FName PropertyName = (PropertyChangedEvent.Property != nullptr) ? PropertyChangedEvent.Property->GetFName() : NAME_None;
+	if (PropertyName == GET_MEMBER_NAME_CHECKED(AEditorUtilityActor, bReceivesEditorInput))
+	{
+		SetReceivesEditorInput(bReceivesEditorInput);
+	}
+}
+#endif
+
+
+void AEditorUtilityActor::SetReceivesEditorInput(bool bInValue)
+{
+	bReceivesEditorInput = bInValue;
+	if (bReceivesEditorInput)
+	{
+		CreateEditorInput();
+	}
+	else
+	{
+		RemoveEditorInput();
+	}
+}
\ No newline at end of file
