Optimized editor utility level editor extensions to scale better with the number of actors selected

diff --git a/Engine/Source/Editor/Blutility/Private/BlutilityLevelEditorExtensions.cpp b/Engine/Source/Editor/Blutility/Private/BlutilityLevelEditorExtensions.cpp
--- a/Engine/Source/Editor/Blutility/Private/BlutilityLevelEditorExtensions.cpp
+++ b/Engine/Source/Editor/Blutility/Private/BlutilityLevelEditorExtensions.cpp
@@ -29,7 +29,6 @@ public:
 		TMap<IEditorUtilityExtension*, TSet<int32>> UtilityAndSelectionIndices;
 
 		// Run thru the actors to determine if any meet our criteria
-		TArray<IEditorUtilityExtension*> SupportedUtils;
 		TArray<AActor*> SupportedActors;
 		if (SelectedActors.Num() > 0)
 		{
@@ -37,66 +36,69 @@ public:
 			TArray<FAssetData> UtilAssets;
 			FBlutilityMenuExtensions::GetBlutilityClasses(UtilAssets, UActorActionUtility::StaticClass()->GetFName());
 
-			// Collect all UActorActionUtility derived (non-blueprint) classes 
-			TSet<UActorActionUtility*> AssetClasses;
+			TMap<UActorActionUtility*, UClass*> ActorActionUtilities;
+
+			// Process asset based utilities
+			for (const FAssetData& UtilAsset : UtilAssets)
+			{
+				if (UEditorUtilityBlueprint* Blueprint = Cast<UEditorUtilityBlueprint>(UtilAsset.GetAsset()))
+				{
+					if (UClass* BPClass = Blueprint->GeneratedClass.Get())
+					{
+						if (UActorActionUtility* ActorActionUtility = Cast<UActorActionUtility>(BPClass->GetDefaultObject()))
+						{
+							UClass*& SupportedClass = ActorActionUtilities.FindOrAdd(ActorActionUtility);
+							if (!SupportedClass)
+							{
+								SupportedClass = ActorActionUtility->GetSupportedClass();
+							}
+						}
+					}
+				}
+			}
+
+			// Process non-asset based utilities
 			for (TObjectIterator<UClass> AssetActionClassIt; AssetActionClassIt; ++AssetActionClassIt)
 			{
 				if (AssetActionClassIt->IsChildOf(UActorActionUtility::StaticClass()) && UActorActionUtility::StaticClass()->GetFName() != AssetActionClassIt->GetFName() && AssetActionClassIt->ClassGeneratedBy == nullptr)
 				{
-					AssetClasses.Add(Cast<UActorActionUtility>(AssetActionClassIt->GetDefaultObject()));
+					if (UActorActionUtility* ActorActionUtility = Cast<UActorActionUtility>(AssetActionClassIt->GetDefaultObject()))
+					{
+						UClass*& SupportedClass = ActorActionUtilities.FindOrAdd(ActorActionUtility);
+						if (!SupportedClass)
+						{
+							SupportedClass = ActorActionUtility->GetSupportedClass();
+						}
+					}
 				}
 			}
 
-			if (UtilAssets.Num() + AssetClasses.Num() > 0)
+			if (UtilAssets.Num() + ActorActionUtilities.Num() > 0)
 			{
 				for (AActor* Actor : SelectedActors)
 				{
 					if (Actor)
 					{
-						auto ProcessAssetAction = [&SupportedUtils, &SupportedActors, &UtilityAndSelectionIndices, Actor](UActorActionUtility* InAction)
+						int32 SupportedActorIndex = INDEX_NONE;
+						for (const TPair<UActorActionUtility*, UClass*>& ActorActionUtilityIt : ActorActionUtilities)
 						{
-							bool bPassesClassFilter = false;
-
-							UClass* SupportedClass = InAction->GetSupportedClass();
+							UClass* SupportedClass = ActorActionUtilityIt.Value;
 							if (SupportedClass == nullptr || (SupportedClass && Actor->GetClass()->IsChildOf(SupportedClass)))
 							{
-								bPassesClassFilter = true;
-							}
-							
-							if (bPassesClassFilter)
-							{
-								SupportedUtils.AddUnique(InAction);
-								const int32 Index = SupportedActors.AddUnique(Actor);
-								UtilityAndSelectionIndices.FindOrAdd(InAction).Add(Index);
-							}
-						};
-						
-						// Process asset based utilities
-						for (const FAssetData& UtilAsset : UtilAssets)
-						{
-							if (UEditorUtilityBlueprint* Blueprint = Cast<UEditorUtilityBlueprint>(UtilAsset.GetAsset()))
-							{
-								if (UClass* BPClass = Blueprint->GeneratedClass.Get())
+								UActorActionUtility* Action = ActorActionUtilityIt.Key;
+								if (SupportedActorIndex == INDEX_NONE)
 								{
-									if (UActorActionUtility* DefaultObject = Cast<UActorActionUtility>(BPClass->GetDefaultObject()))
-									{
-										ProcessAssetAction(DefaultObject);
-									}
+									SupportedActorIndex = SupportedActors.Add(Actor);
 								}
+								UtilityAndSelectionIndices.FindOrAdd(Action).Add(SupportedActorIndex);
 							}
 						}
-
-						// Process non-asset based utilities
-						for (UActorActionUtility* Action : AssetClasses)
-						{
-							ProcessAssetAction(Action);
-						}
 					}
 				}
 			}
 		}
 
-		if (SupportedUtils.Num() > 0)
+		if (SupportedActors.Num() > 0)
 		{
 			// Add asset actions extender
 			Extender->AddMenuExtension(
