Adding a function that allows to specify ID to spawn a tab. Allows opening multiple copies of the same Editor Utility Widget.

diff --git a/Engine/Source/Editor/Blutility/Private/EditorUtilitySubsystem.cpp b/Engine/Source/Editor/Blutility/Private/EditorUtilitySubsystem.cpp
--- a/Engine/Source/Editor/Blutility/Private/EditorUtilitySubsystem.cpp
+++ b/Engine/Source/Editor/Blutility/Private/EditorUtilitySubsystem.cpp
@@ -183,11 +183,16 @@ UEditorUtilityWidget* UEditorUtilitySubsystem::SpawnAndRegisterTab(class UEditor
 	return SpawnAndRegisterTabAndGetID(InBlueprint, InTabID);
 }
 
+UEditorUtilityWidget* UEditorUtilitySubsystem::SpawnAndRegisterTabWithId(class UEditorUtilityWidgetBlueprint* InBlueprint, FName InTabID)
+{
+	return SpawnAndRegisterTabAndGetID(InBlueprint, InTabID);;
+}
+
 void UEditorUtilitySubsystem::RegisterTabAndGetID(class UEditorUtilityWidgetBlueprint* InBlueprint, FName& NewTabID)
 {
 	if (InBlueprint && !IsRunningCommandlet())
 	{
-		FName RegistrationName = FName(*(InBlueprint->GetPathName() + LOCTEXT("ActiveTabSuffix", "_ActiveTab").ToString()));
+		FName RegistrationName = NewTabID.IsNone() ? FName(*(InBlueprint->GetPathName() + LOCTEXT("ActiveTabSuffix", "_ActiveTab").ToString())) : FName(*(InBlueprint->GetPathName() + NewTabID.ToString()));
 		FText DisplayName = FText::FromString(InBlueprint->GetName());
 		FLevelEditorModule& LevelEditorModule = FModuleManager::GetModuleChecked<FLevelEditorModule>(TEXT("LevelEditor"));
 		TSharedPtr<FTabManager> LevelEditorTabManager = LevelEditorModule.GetLevelEditorTabManager();
diff --git a/Engine/Source/Editor/Blutility/Public/EditorUtilitySubsystem.h b/Engine/Source/Editor/Blutility/Public/EditorUtilitySubsystem.h
--- a/Engine/Source/Editor/Blutility/Public/EditorUtilitySubsystem.h
+++ b/Engine/Source/Editor/Blutility/Public/EditorUtilitySubsystem.h
@@ -59,6 +59,12 @@ public:
 	UFUNCTION(BlueprintCallable, Category = "Development|Editor")
 	UEditorUtilityWidget* SpawnAndRegisterTab(class UEditorUtilityWidgetBlueprint* InBlueprint);
 
+	/**
+	 * Unlike SpawnAndRegisterTabAndGetID allows spawn tab while providing TabID from Python scripts or BP
+	 */
+	UFUNCTION(BlueprintCallable, Category = "Development|Editor")
+	UEditorUtilityWidget* SpawnAndRegisterTabWithId(class UEditorUtilityWidgetBlueprint* InBlueprint, FName InTabID);
+
 	UFUNCTION(BlueprintCallable, Category = "Development|Editor")
 	void RegisterTabAndGetID(class UEditorUtilityWidgetBlueprint* InBlueprint, FName& NewTabID);
 
