Scripted Asset Actions now work with collection parameters.

CPF_Parm was not being applied to child properties (ie. items in the array), and the IsPropertyVisible delegate wasn't checking parent properties for the flag, which was causing only the root to be visible. This is a result of a change made to the details view to run the IsPropertyVisible delegate more frequently - previously it was only run when initializing the property tree.

diff --git a/Engine/Source/Editor/Blutility/Private/BlutilityMenuExtensions.cpp b/Engine/Source/Editor/Blutility/Private/BlutilityMenuExtensions.cpp
--- a/Engine/Source/Editor/Blutility/Private/BlutilityMenuExtensions.cpp
+++ b/Engine/Source/Editor/Blutility/Private/BlutilityMenuExtensions.cpp
@@ -74,7 +74,25 @@ class SFunctionParamDialog : public SCompoundWidget
 		// Hide any property that has been marked as such
 		StructureDetailsView->GetDetailsView()->SetIsPropertyVisibleDelegate(FIsPropertyVisible::CreateLambda([HiddenPropertyName](const FPropertyAndParent& InPropertyAndParent)
 		{
-			return InPropertyAndParent.Property.HasAnyPropertyFlags(CPF_Parm) && InPropertyAndParent.Property.GetFName() != HiddenPropertyName;
+			if (InPropertyAndParent.Property.GetFName() == HiddenPropertyName)
+			{
+				return false;
+			}
+
+			if (InPropertyAndParent.Property.HasAnyPropertyFlags(CPF_Parm))
+			{
+				return true;
+			}
+
+			for (const FProperty* Parent : InPropertyAndParent.ParentProperties)
+			{
+				if (Parent->HasAnyPropertyFlags(CPF_Parm))
+				{
+					return true;
+				}
+			}
+
+			return false;
 		}));
 
 		StructureDetailsView->GetDetailsView()->ForceRefresh();
