Fixed issue where nested utility widgets could cause the open level to be dirtied since they wouldn't properly be marked transient like the outermost utility widget

diff --git a/Engine/Source/Editor/Blutility/Private/EditorUtilityWidgetBlueprint.cpp b/Engine/Source/Editor/Blutility/Private/EditorUtilityWidgetBlueprint.cpp
--- a/Engine/Source/Editor/Blutility/Private/EditorUtilityWidgetBlueprint.cpp
+++ b/Engine/Source/Editor/Blutility/Private/EditorUtilityWidgetBlueprint.cpp
@@ -7,6 +7,8 @@
 #include "IBlutilityModule.h"
 #include "Modules/ModuleManager.h"
 #include "LevelEditor.h"
+#include "Blueprint/UserWidget.h"
+#include "Blueprint/WidgetTree.h"
 
 
 
@@ -82,6 +84,19 @@ TSharedRef<SWidget> UEditorUtilityWidgetBlueprint::CreateUtilityWidget()
 		{
 			// Editor Utility is flagged as transient to prevent from dirty the World it's created in when a property added to the Utility Widget is changed
 			CreatedUMGWidget->SetFlags(RF_Transient);
+
+			// Also mark nested utility widgets as transient to prevent them from dirtying the world (since they'll be created via CreateWidget and not CreateUtilityWidget)
+			TArray<UWidget*> AllWidgets;
+			CreatedUMGWidget->WidgetTree->GetAllWidgets(AllWidgets);
+
+			for (UWidget* Widget : AllWidgets)
+			{
+				if (Widget->IsA(UEditorUtilityWidget::StaticClass()))
+				{
+					Widget->SetFlags(RF_Transient);
+					Widget->Slot->SetFlags(RF_Transient);
+				}
+			}
 		}
 	}
 
