PR #9114: Use the Display Name for Editor Utility Widgets in the menu (Contributed by [REDACTED_PERSON])

diff --git a/Engine/Source/Editor/Blutility/Private/BlutilityModule.cpp b/Engine/Source/Editor/Blutility/Private/BlutilityModule.cpp
--- a/Engine/Source/Editor/Blutility/Private/BlutilityModule.cpp
+++ b/Engine/Source/Editor/Blutility/Private/BlutilityModule.cpp
@@ -111,7 +111,7 @@ public:
 						const UEditorUtilityWidget* CDO = Blueprint->GeneratedClass->GetDefaultObject<UEditorUtilityWidget>();
 						FName RegistrationName = FName(*(Blueprint->GetPathName() + LOCTEXT("ActiveTabSuffix", "_ActiveTab").ToString()));
 						Blueprint->SetRegistrationName(RegistrationName);
-						FText DisplayName = FText::FromString(Blueprint->GetName());
+						FText DisplayName = FText::FromString(FName::NameToDisplayString(Blueprint->GetName(), false));
 						if (LevelEditorTabManager && !LevelEditorTabManager->HasTabSpawner(RegistrationName))
 						{
 							LevelEditorTabManager->RegisterTabSpawner(RegistrationName, FOnSpawnTab::CreateUObject(Blueprint, &UEditorUtilityWidgetBlueprint::SpawnEditorUITab))
@@ -311,4 +311,4 @@ protected:
 
 
 IMPLEMENT_MODULE( FBlutilityModule, Blutility );
-#undef LOCTEXT_NAMESPACE 
\ No newline at end of file
+#undef LOCTEXT_NAMESPACE 
diff --git a/Engine/Source/Editor/Blutility/Private/EditorUtilitySubsystem.cpp b/Engine/Source/Editor/Blutility/Private/EditorUtilitySubsystem.cpp
--- a/Engine/Source/Editor/Blutility/Private/EditorUtilitySubsystem.cpp
+++ b/Engine/Source/Editor/Blutility/Private/EditorUtilitySubsystem.cpp
@@ -193,7 +193,7 @@ void UEditorUtilitySubsystem::RegisterTabAndGetID(class UEditorUtilityWidgetBlue
 	if (InBlueprint && !IsRunningCommandlet())
 	{
 		FName RegistrationName = NewTabID.IsNone() ? FName(*(InBlueprint->GetPathName() + LOCTEXT("ActiveTabSuffix", "_ActiveTab").ToString())) : FName(*(InBlueprint->GetPathName() + NewTabID.ToString()));
-		FText DisplayName = FText::FromString(InBlueprint->GetName());
+		FText DisplayName = FText::FromString(FName::NameToDisplayString(InBlueprint->GetName(), false));
 		FLevelEditorModule& LevelEditorModule = FModuleManager::GetModuleChecked<FLevelEditorModule>(TEXT("LevelEditor"));
 		TSharedPtr<FTabManager> LevelEditorTabManager = LevelEditorModule.GetLevelEditorTabManager();
 		if (!LevelEditorTabManager->HasTabSpawner(RegistrationName))
