Avoid loading Non-BP assets (GetAsset loads)

diff --git a/Engine/Source/Editor/Blutility/Private/BlutilityContentBrowserExtensions.cpp b/Engine/Source/Editor/Blutility/Private/BlutilityContentBrowserExtensions.cpp
--- a/Engine/Source/Editor/Blutility/Private/BlutilityContentBrowserExtensions.cpp
+++ b/Engine/Source/Editor/Blutility/Private/BlutilityContentBrowserExtensions.cpp
@@ -18,6 +18,7 @@
 #include "AssetRegistry/AssetRegistryModule.h"
 #include "EditorUtilityBlueprint.h"
 #include "Framework/Application/SlateApplication.h"
+#include "Templates/SubclassOf.h"
 
 #include "BlueprintEditorModule.h"
 #include "BlutilityMenuExtensions.h"
@@ -67,10 +68,18 @@ public:
 						bool bPassesClassFilter = false;
 						if (bIsActionForBlueprints)
 						{
-							if (UBlueprint* AssetAsBlueprint = Cast<UBlueprint>(Asset.GetAsset()))
+							if (TSubclassOf<UBlueprint> AssetClass = Asset.GetClass())
 							{
-								// It's a blueprint, but is it the right kind?
-								bPassesClassFilter = AssetAsBlueprint->ParentClass && AssetAsBlueprint->ParentClass->IsChildOf(SupportedClass);
+								if (UBlueprint* AssetAsBlueprint = Cast<UBlueprint>(Asset.GetAsset()))
+								{
+									// It's a blueprint, but is it the right kind?
+									bPassesClassFilter = AssetAsBlueprint->ParentClass && AssetAsBlueprint->ParentClass->IsChildOf(SupportedClass);
+								}
+								else
+								{
+									// Not a blueprint
+									bPassesClassFilter = false;
+								}
 							}
 							else
 							{
diff --git a/Engine/Source/Editor/Blutility/Private/EditorUtilityLibrary.cpp b/Engine/Source/Editor/Blutility/Private/EditorUtilityLibrary.cpp
--- a/Engine/Source/Editor/Blutility/Private/EditorUtilityLibrary.cpp
+++ b/Engine/Source/Editor/Blutility/Private/EditorUtilityLibrary.cpp
@@ -11,7 +11,7 @@
 #include "IAssetTools.h"
 #include "EditorUtilitySubsystem.h"
 #include "Kismet/KismetSystemLibrary.h"
-
+#include "Templates/SubclassOf.h"
 
 #define LOCTEXT_NAMESPACE "BlutilityLevelEditorExtensions"
 
@@ -263,9 +263,12 @@ TArray<UClass*> UEditorUtilityLibrary::GetSelectedBlueprintClasses()
 	TArray<UClass*> Result;
 	for (FAssetData& AssetData : SelectedAssets)
 	{
-		if (UBlueprint* Blueprint = Cast<UBlueprint>(AssetData.GetAsset()))
+		if (TSubclassOf<UBlueprint> AssetClass = AssetData.GetClass())
 		{
-			Result.Add(Blueprint->GeneratedClass);
+			if (UBlueprint* Blueprint = Cast<UBlueprint>(AssetData.GetAsset()))
+			{
+				Result.Add(Blueprint->GeneratedClass);
+			}
 		}
 	}
 
