#ue5 - Fix type truncation warnings in Blutility, CollectionManagerTypes modules.

diff --git a/Engine/Source/Developer/CollectionManager/Public/CollectionManagerTypes.h b/Engine/Source/Developer/CollectionManager/Public/CollectionManagerTypes.h
--- a/Engine/Source/Developer/CollectionManager/Public/CollectionManagerTypes.h
+++ b/Engine/Source/Developer/CollectionManager/Public/CollectionManagerTypes.h
@@ -186,7 +186,7 @@ struct ECollectionShareType
 /** Controls how the collections manager will recurse when performing work against a given collection */
 struct ECollectionRecursionFlags
 {
-	typedef uint8 Flags;
+	typedef uint32 Flags;
 	enum Flag
 	{
 		/** Include the current collection when performing work */
diff --git a/Engine/Source/Editor/Blutility/Private/EditorUtilityLibrary.cpp b/Engine/Source/Editor/Blutility/Private/EditorUtilityLibrary.cpp
--- a/Engine/Source/Editor/Blutility/Private/EditorUtilityLibrary.cpp
+++ b/Engine/Source/Editor/Blutility/Private/EditorUtilityLibrary.cpp
@@ -210,7 +210,7 @@ TArray<AActor*> UEditorUtilityLibrary::GetSelectionSet()
 	return Result;
 }
 
-void UEditorUtilityLibrary::GetSelectionBounds(FVector& Origin, FVector& BoxExtent, float& SphereRadius)
+void UEditorUtilityLibrary::GetSelectionBounds(FVector& Origin, FVector& BoxExtent, double& SphereRadius)
 {
 	bool bFirstItem = true;
 
diff --git a/Engine/Source/Editor/Blutility/Private/GlobalEditorUtilityBase.cpp b/Engine/Source/Editor/Blutility/Private/GlobalEditorUtilityBase.cpp
--- a/Engine/Source/Editor/Blutility/Private/GlobalEditorUtilityBase.cpp
+++ b/Engine/Source/Editor/Blutility/Private/GlobalEditorUtilityBase.cpp
@@ -64,7 +64,7 @@ void UDEPRECATED_GlobalEditorUtilityBase::GetSelectionBounds(FVector& Origin, FV
 
 	Origin = Extents.Origin;
 	BoxExtent = Extents.BoxExtent;
-	SphereRadius = Extents.SphereRadius;
+	SphereRadius = (float)Extents.SphereRadius;
 }
 
 void UDEPRECATED_GlobalEditorUtilityBase::ForEachSelectedActor()
diff --git a/Engine/Source/Editor/Blutility/Public/EditorUtilityLibrary.h b/Engine/Source/Editor/Blutility/Public/EditorUtilityLibrary.h
--- a/Engine/Source/Editor/Blutility/Public/EditorUtilityLibrary.h
+++ b/Engine/Source/Editor/Blutility/Public/EditorUtilityLibrary.h
@@ -143,7 +143,7 @@ public:
 	static TArray<AActor*> GetSelectionSet();
 
 	UFUNCTION(BlueprintCallable, Category = "Development|Editor")
-	static void GetSelectionBounds(FVector& Origin, FVector& BoxExtent, float& SphereRadius);
+	static void GetSelectionBounds(FVector& Origin, FVector& BoxExtent, double& SphereRadius);
 
 	// Gets the set of currently selected assets
 	UFUNCTION(BlueprintCallable, Category = "Development|Editor")
