Add nullptr check for engine shutdown.

#rnx

diff --git a/Engine/Source/Editor/Blutility/Private/BlutilityModule.cpp b/Engine/Source/Editor/Blutility/Private/BlutilityModule.cpp
--- a/Engine/Source/Editor/Blutility/Private/BlutilityModule.cpp
+++ b/Engine/Source/Editor/Blutility/Private/BlutilityModule.cpp
@@ -305,6 +305,10 @@ protected:
 	void HandleAssetRemoved(const FAssetData& InAssetData)
 	{
 		bool bDeletingLoadedUI = false;
+		if (!GEditor || !GEditor->GetEditorSubsystem<UEditorUtilitySubsystem>())
+		{
+			return;
+		}
 		for (const FSoftObjectPath& LoadedUIPath : GEditor->GetEditorSubsystem<UEditorUtilitySubsystem>()->LoadedUIs)
 		{
 			if (LoadedUIPath.GetAssetPathName() == InAssetData.ObjectPath)
