#ue5 - Revert function param change in BP GetSelectionBounds() because C++ uses will complain. Need to deprecate and replace with another method for BPs.

diff --git a/Engine/Source/Editor/Blutility/Private/EditorUtilityLibrary.cpp b/Engine/Source/Editor/Blutility/Private/EditorUtilityLibrary.cpp
--- a/Engine/Source/Editor/Blutility/Private/EditorUtilityLibrary.cpp
+++ b/Engine/Source/Editor/Blutility/Private/EditorUtilityLibrary.cpp
@@ -210,7 +210,7 @@ TArray<AActor*> UEditorUtilityLibrary::GetSelectionSet()
 	return Result;
 }
 
-void UEditorUtilityLibrary::GetSelectionBounds(FVector& Origin, FVector& BoxExtent, double& SphereRadius)
+void UEditorUtilityLibrary::GetSelectionBounds(FVector& Origin, FVector& BoxExtent, float& SphereRadius)
 {
 	bool bFirstItem = true;
 
@@ -234,7 +234,7 @@ void UEditorUtilityLibrary::GetSelectionBounds(FVector& Origin, FVector& BoxExte
 
 	Origin = Extents.Origin;
 	BoxExtent = Extents.BoxExtent;
-	SphereRadius = Extents.SphereRadius;
+	SphereRadius = (float)Extents.SphereRadius; // TODO: LWC: should be double, but need to deprecate function and replace for old C++ references to continue working.
 }
 
 TArray<UObject*> UEditorUtilityLibrary::GetSelectedAssets()
diff --git a/Engine/Source/Editor/Blutility/Public/EditorUtilityLibrary.h b/Engine/Source/Editor/Blutility/Public/EditorUtilityLibrary.h
--- a/Engine/Source/Editor/Blutility/Public/EditorUtilityLibrary.h
+++ b/Engine/Source/Editor/Blutility/Public/EditorUtilityLibrary.h
@@ -143,7 +143,7 @@ public:
 	static TArray<AActor*> GetSelectionSet();
 
 	UFUNCTION(BlueprintCallable, Category = "Development|Editor")
-	static void GetSelectionBounds(FVector& Origin, FVector& BoxExtent, double& SphereRadius);
+	static void GetSelectionBounds(FVector& Origin, FVector& BoxExtent, float& SphereRadius);
 
 	// Gets the set of currently selected assets
 	UFUNCTION(BlueprintCallable, Category = "Development|Editor")
