Editor Utility Library: Add a function to expose syncing to the current content browser folder to scripts

diff --git a/Engine/Source/Editor/Blutility/Private/EditorUtilityLibrary.cpp b/Engine/Source/Editor/Blutility/Private/EditorUtilityLibrary.cpp
--- a/Engine/Source/Editor/Blutility/Private/EditorUtilityLibrary.cpp
+++ b/Engine/Source/Editor/Blutility/Private/EditorUtilityLibrary.cpp
@@ -327,6 +327,12 @@ TArray<FString> UEditorUtilityLibrary::GetSelectedFolderPaths()
 	return Paths;
 }
 
+void UEditorUtilityLibrary::SyncBrowserToFolders(const TArray<FString>& FolderList)
+{
+	FContentBrowserModule& ContentBrowserModule = FModuleManager::Get().LoadModuleChecked<FContentBrowserModule>("ContentBrowser");
+	ContentBrowserModule.Get().SyncBrowserToFolders( FolderList, false, true );
+}
+
 #endif
 
 #undef LOCTEXT_NAMESPACE
diff --git a/Engine/Source/Editor/Blutility/Public/EditorUtilityLibrary.h b/Engine/Source/Editor/Blutility/Public/EditorUtilityLibrary.h
--- a/Engine/Source/Editor/Blutility/Public/EditorUtilityLibrary.h
+++ b/Engine/Source/Editor/Blutility/Public/EditorUtilityLibrary.h
@@ -180,5 +180,12 @@ public:
 	// Gets the path to the currently selected folder in the content browser
 	UFUNCTION(BlueprintPure, Category = "Development|Editor")
 	static TArray<FString> GetSelectedFolderPaths();
+
+	/**
+	 * Sync the Content Browser to the given folder(s)
+	 * @param	FolderList	The list of folders to sync to in the Content Browser
+	 */
+	UFUNCTION(BlueprintCallable, Category = "Editor Scripting | Content Browser")
+	static void SyncBrowserToFolders(const TArray<FString>& FolderList);
 #endif
 };
