If editing an EditorUtilityWidget, add button to run / stop utility widget

diff --git a/Engine/Source/Editor/Blutility/Private/BlutilityModule.cpp b/Engine/Source/Editor/Blutility/Private/BlutilityModule.cpp
--- a/Engine/Source/Editor/Blutility/Private/BlutilityModule.cpp
+++ b/Engine/Source/Editor/Blutility/Private/BlutilityModule.cpp
@@ -9,6 +9,7 @@
 #include "AssetTypeCategories.h"
 #include "BlutilityContentBrowserExtensions.h"
 #include "BlutilityLevelEditorExtensions.h"
+#include "BlutilityUMGEditorExtensions.h"
 #include "Containers/Array.h"
 #include "Containers/UnrealString.h"
 #include "Delegates/Delegate.h"
@@ -101,6 +102,7 @@ public:
 
 		FBlutilityContentBrowserExtensions::InstallHooks();
 		FBlutilityLevelEditorExtensions::InstallHooks();
+		FBlutilityUMGEditorExtensions::InstallHooks();
 
 		ScriptedEditorWidgetsGroup = WorkspaceMenu::GetMenuStructure().GetToolsCategory()->AddGroup(
 			LOCTEXT("WorkspaceMenu_EditorUtilityWidgetsGroup", "Editor Utility Widgets"),
@@ -204,6 +206,7 @@ public:
 
 		FBlutilityLevelEditorExtensions::RemoveHooks();
 		FBlutilityContentBrowserExtensions::RemoveHooks();
+		FBlutilityUMGEditorExtensions::RemoveHooks();
 
 		// Only unregister if the asset tools module is loaded.  We don't want to forcibly load it during shutdown phase.
 		check( EditorBlueprintAssetTypeActions.IsValid() );
diff --git a/Engine/Source/Editor/Blutility/Private/BlutilityUMGEditorExtensions.cpp b/Engine/Source/Editor/Blutility/Private/BlutilityUMGEditorExtensions.cpp
new file mode 100644
--- /dev/null
+++ b/Engine/Source/Editor/Blutility/Private/BlutilityUMGEditorExtensions.cpp
@@ -0,0 +1,184 @@
+// Copyright Epic Games, Inc. All Rights Reserved.
+
+#include "BlutilityUMGEditorExtensions.h"
+
+#include "IDocumentation.h"
+#include "BlueprintEditor.h"
+#include "BlueprintEditorContext.h"
+#include "Blueprint/UserWidget.h"
+#include "BlueprintModes/WidgetBlueprintApplicationModes.h"
+#include "Editor.h"
+#include "Editor/EditorEngine.h"
+#include "EditorUtilitySubsystem.h"
+#include "EditorUtilityWidget.h"
+#include "EditorUtilityWidgetBlueprint.h"
+#include "Modules/ModuleManager.h"
+#include "ToolMenus.h"
+#include "UMGEditorModule.h"
+#include "WidgetBlueprint.h"
+#include "WidgetBlueprintEditor.h"
+
+class FUICommandList;
+
+#define LOCTEXT_NAMESPACE "BlutilityUMGEditorExtensions"
+
+namespace UE::Blutility::Private
+{
+	// Hardcoded toolbar name since not exposed without instance of UMG editor
+	static FName WidgetBlueprintEditorToolbarName = FName("AssetEditor.WidgetBlueprintEditor.ToolBar");
+	static FName BlutilityExtensionName = FName("BlutilityExtensions");
+}
+
+class FBlutilityUMGEditorExtensions_Impl
+{
+public:
+
+	static void AddEditorUtilityExtension(UToolMenu* InMenu)
+	{
+		FToolMenuSection& Section = InMenu->AddSection(UE::Blutility::Private::BlutilityExtensionName);
+		Section.InsertPosition = FToolMenuInsert("Compile", EToolMenuInsertType::After);
+		Section.AddDynamicEntry("RunStopUtilityWidget", FNewToolMenuSectionDelegate::CreateLambda([](FToolMenuSection& InSection)
+		{
+			UBlueprintEditorToolMenuContext* Context = InSection.FindContext<UBlueprintEditorToolMenuContext>();
+			if (Context)
+			{
+				if (TSharedPtr<FWidgetBlueprintEditor> WidgetBlueprintEditor = StaticCastSharedPtr<FWidgetBlueprintEditor>(Context->BlueprintEditor.Pin()))
+				{
+					if (WidgetBlueprintEditor->GetPreview() && WidgetBlueprintEditor->GetPreview()->IsEditorUtility())
+					{
+						TWeakPtr<FWidgetBlueprintEditor> WeakWidgetBlueprintEditor = WidgetBlueprintEditor;
+
+						FName UtilityWidgetName = WidgetBlueprintEditor->GetBlueprintObj()->GetFName();
+
+						// Add Run Button, will toggle if tab already ran.
+						InSection.AddEntry(FToolMenuEntry::InitToolBarButton(
+							"RunUtilityWidget",
+							FUIAction(
+								FExecuteAction::CreateLambda([WeakWidgetBlueprintEditor]()
+									{ 
+										if (TSharedPtr<FWidgetBlueprintEditor> WidgetBlueprintEditor = WeakWidgetBlueprintEditor.Pin())
+										{
+											if (UEditorUtilityWidgetBlueprint* EditorWidget = Cast<UEditorUtilityWidgetBlueprint>(WidgetBlueprintEditor->GetBlueprintObj()))
+											{
+												GEditor->GetEditorSubsystem<UEditorUtilitySubsystem>()->SpawnAndRegisterTabWithId(EditorWidget, EditorWidget->GetFName());
+											}
+										}
+									}),
+								FCanExecuteAction(),
+								FGetActionCheckState(),
+								FIsActionButtonVisible::CreateLambda([WeakWidgetBlueprintEditor]()
+									{
+										if (TSharedPtr<FWidgetBlueprintEditor> WidgetBlueprintEditor = WeakWidgetBlueprintEditor.Pin())
+										{
+											if (UEditorUtilityWidgetBlueprint* EditorWidget = Cast<UEditorUtilityWidgetBlueprint>(WidgetBlueprintEditor->GetBlueprintObj()))
+											{
+												// Editor Utility Prepends path on Tab Name
+												FName UtilityTabName = FName(*(EditorWidget->GetPathName() + EditorWidget->GetFName().ToString()));
+												return !GEditor->GetEditorSubsystem<UEditorUtilitySubsystem>()->DoesTabExist(UtilityTabName);
+											}
+										}
+										return false; 
+									})
+							)
+							, LOCTEXT("RunUtilityWidget", "Run Utility Widget")
+							, LOCTEXT("RunUtilityWidgetToolTip", "Run Utility Widget, tab name will be based off of asset.")
+							, FSlateIcon(FAppStyle::Get().GetStyleSetName(), "PlayWorld.PlayInViewport") // @TODO: DarenC - Placeholder
+						));
+
+						// Add Stop Button, will toggle if tab already ran.
+						InSection.AddEntry(FToolMenuEntry::InitToolBarButton(
+							"StopUtilityWidget",
+							FUIAction(
+								FExecuteAction::CreateLambda([WeakWidgetBlueprintEditor]()
+									{
+										if (TSharedPtr<FWidgetBlueprintEditor> WidgetBlueprintEditor = WeakWidgetBlueprintEditor.Pin())
+										{
+											if (UEditorUtilityWidgetBlueprint* EditorWidget = Cast<UEditorUtilityWidgetBlueprint>(WidgetBlueprintEditor->GetBlueprintObj()))
+											{
+												// Editor Utility Prepends path on Tab Name
+												FName UtilityTabName = FName(*(EditorWidget->GetPathName() + EditorWidget->GetFName().ToString()));
+												GEditor->GetEditorSubsystem<UEditorUtilitySubsystem>()->CloseTabByID(UtilityTabName);
+											}
+										}
+									}),
+								FCanExecuteAction(),
+								FGetActionCheckState(),
+								FIsActionButtonVisible::CreateLambda([WeakWidgetBlueprintEditor]()
+									{
+										if (TSharedPtr<FWidgetBlueprintEditor> WidgetBlueprintEditor = WeakWidgetBlueprintEditor.Pin())
+										{
+											if (UEditorUtilityWidgetBlueprint* EditorWidget = Cast<UEditorUtilityWidgetBlueprint>(WidgetBlueprintEditor->GetBlueprintObj()))
+											{
+												// Editor Utility Prepends path on Tab Name
+												FName UtilityTabName = FName(*(EditorWidget->GetPathName() + EditorWidget->GetFName().ToString()));
+												return GEditor->GetEditorSubsystem<UEditorUtilitySubsystem>()->DoesTabExist(UtilityTabName);
+											}
+										}
+										return false; 
+									})
+							)
+							, LOCTEXT("StopUtilityWidget", "Stop Utility Widget")
+							, LOCTEXT("StopUtilityWidgetToolTip", "Stop Utility Widget, will stop tab based off of asset name (Does not stop manually ran instances).")
+							, FSlateIcon(FAppStyle::Get().GetStyleSetName(), "PlayWorld.StopPlaySession") // @TODO: DarenC - Placeholder
+						));
+					}
+				}
+			}
+		}));
+	}
+
+	static void RemoveEditorUtilityExtension(UToolMenu* InMenu)
+	{
+		InMenu->RemoveSection(UE::Blutility::Private::BlutilityExtensionName);
+	}
+};
+
+void FBlutilityUMGEditorExtensions::InstallHooks()
+{
+	// Load module to ensure relevant toolbar exists
+	IUMGEditorModule& UMGEditorModule = FModuleManager::Get().LoadModuleChecked<IUMGEditorModule>("UMGEditor");
+
+	auto AddExtensionsUMGEditorToolbar = [](const FName InModeName, FName ParentToolbarName)
+	{
+		const FName ModeSpecificToolbarName = *(ParentToolbarName.ToString() + TEXT(".") + InModeName.ToString());
+
+		if (!UToolMenus::Get()->IsMenuRegistered(ModeSpecificToolbarName))
+		{
+			UToolMenus::Get()->RegisterMenu(ModeSpecificToolbarName, ParentToolbarName, EMultiBoxType::ToolBar);
+		}
+
+		if (UToolMenu* Toolbar = UToolMenus::Get()->FindMenu(ModeSpecificToolbarName))
+		{
+			FBlutilityUMGEditorExtensions_Impl::AddEditorUtilityExtension(Toolbar);
+		}
+	};
+
+	AddExtensionsUMGEditorToolbar(FWidgetBlueprintApplicationModes::DesignerMode, UE::Blutility::Private::WidgetBlueprintEditorToolbarName);
+	AddExtensionsUMGEditorToolbar(FWidgetBlueprintApplicationModes::GraphMode, UE::Blutility::Private::WidgetBlueprintEditorToolbarName);
+}
+
+void FBlutilityUMGEditorExtensions::RemoveHooks()
+{
+	// Load module to ensure relevant toolbar exists
+	IUMGEditorModule& UMGEditorModule = FModuleManager::Get().LoadModuleChecked<IUMGEditorModule>("UMGEditor");
+
+	auto RemoveExtensionsUMGEditorToolbar = [](const FName InModeName, FName ParentToolbarName)
+	{
+		const FName ModeSpecificToolbarName = *(ParentToolbarName.ToString() + TEXT(".") + InModeName.ToString());
+
+		if (!UToolMenus::Get()->IsMenuRegistered(ModeSpecificToolbarName))
+		{
+			UToolMenus::Get()->RegisterMenu(ModeSpecificToolbarName, ParentToolbarName, EMultiBoxType::ToolBar);
+		}
+
+		if (UToolMenu* Toolbar = UToolMenus::Get()->FindMenu(ModeSpecificToolbarName))
+		{
+			FBlutilityUMGEditorExtensions_Impl::RemoveEditorUtilityExtension(Toolbar);
+		}
+	};
+
+	RemoveExtensionsUMGEditorToolbar(FWidgetBlueprintApplicationModes::DesignerMode, UE::Blutility::Private::WidgetBlueprintEditorToolbarName);
+	RemoveExtensionsUMGEditorToolbar(FWidgetBlueprintApplicationModes::GraphMode, UE::Blutility::Private::WidgetBlueprintEditorToolbarName);
+}
+
+#undef LOCTEXT_NAMESPACE
diff --git a/Engine/Source/Editor/Blutility/Private/BlutilityUMGEditorExtensions.h b/Engine/Source/Editor/Blutility/Private/BlutilityUMGEditorExtensions.h
new file mode 100644
--- /dev/null
+++ b/Engine/Source/Editor/Blutility/Private/BlutilityUMGEditorExtensions.h
@@ -0,0 +1,11 @@
+// Copyright Epic Games, Inc. All Rights Reserved.
+
+#pragma once
+
+// Integrate Blutility actions associated with UMG editor functionality (e.g. Run / Stop Utility Widget)
+class FBlutilityUMGEditorExtensions
+{
+public:
+	static void InstallHooks();
+	static void RemoveHooks();
+};
diff --git a/Engine/Source/Editor/UMGEditor/Private/BlueprintModes/WidgetDesignerApplicationMode.cpp b/Engine/Source/Editor/UMGEditor/Private/BlueprintModes/WidgetDesignerApplicationMode.cpp
--- a/Engine/Source/Editor/UMGEditor/Private/BlueprintModes/WidgetDesignerApplicationMode.cpp
+++ b/Engine/Source/Editor/UMGEditor/Private/BlueprintModes/WidgetDesignerApplicationMode.cpp
@@ -155,8 +155,11 @@ FWidgetDesignerApplicationMode::FWidgetDesignerApplicationMode(TSharedPtr<FWidge
 	ToolbarExtender = UMGEditorModule.GetToolBarExtensibilityManager()->GetAllExtenders();
 	
 	InWidgetEditor->GetWidgetToolbarBuilder()->AddWidgetBlueprintEditorModesToolbar(ToolbarExtender);
+	InWidgetEditor->RegisterModeToolbarIfUnregistered(GetModeName());
 
-	if (UToolMenu* Toolbar = InWidgetEditor->RegisterModeToolbarIfUnregistered(GetModeName()))
+	FName OutParentToolbarName;
+	FName ToolBarname = InWidgetEditor->GetToolMenuToolbarNameForMode(GetModeName(), OutParentToolbarName);
+	if (UToolMenu* Toolbar = UToolMenus::Get()->FindMenu(ToolBarname))
 	{
 		InWidgetEditor->GetWidgetToolbarBuilder()->AddWidgetReflector(Toolbar);
 		InWidgetEditor->GetWidgetToolbarBuilder()->AddToolPalettes(Toolbar);
diff --git a/Engine/Source/Editor/UMGEditor/Private/BlueprintModes/WidgetGraphApplicationMode.cpp b/Engine/Source/Editor/UMGEditor/Private/BlueprintModes/WidgetGraphApplicationMode.cpp
--- a/Engine/Source/Editor/UMGEditor/Private/BlueprintModes/WidgetGraphApplicationMode.cpp
+++ b/Engine/Source/Editor/UMGEditor/Private/BlueprintModes/WidgetGraphApplicationMode.cpp
@@ -82,8 +82,11 @@ FWidgetGraphApplicationMode::FWidgetGraphApplicationMode(TSharedPtr<FWidgetBluep
 
 	ToolbarExtender = MakeShareable(new FExtender);
 	InWidgetEditor->GetWidgetToolbarBuilder()->AddWidgetBlueprintEditorModesToolbar(ToolbarExtender);
+	InWidgetEditor->RegisterModeToolbarIfUnregistered(GetModeName());
 
-	if (UToolMenu* Toolbar = InWidgetEditor->RegisterModeToolbarIfUnregistered(GetModeName()))
+	FName OutParentToolbarName;
+	FName ToolBarname = InWidgetEditor->GetToolMenuToolbarNameForMode(GetModeName(), OutParentToolbarName);
+	if (UToolMenu* Toolbar = UToolMenus::Get()->FindMenu(ToolBarname))
 	{
 		InWidgetEditor->GetToolbarBuilder()->AddCompileToolbar(Toolbar);
 		InWidgetEditor->GetToolbarBuilder()->AddScriptingToolbar(Toolbar);
