Editor - Blutility - AssetActionUtilities and ActorActionUtilties no longer load on right click in the Content Browser.  Their information is now cached in tag data and we don't load them for real until we actually run the menu options.  The GetSupportedClass() has been deprecated, there is now an array of classes you can provide in the Class Defaults. We will now exclude blutilities in developer folders that are not yours.

This is part 1 of a larger refactor to try and stop loading content on right click in the content browser.

diff --git a/Engine/Source/Editor/Blutility/Blutility.Build.cs b/Engine/Source/Editor/Blutility/Blutility.Build.cs
--- a/Engine/Source/Editor/Blutility/Blutility.Build.cs
+++ b/Engine/Source/Editor/Blutility/Blutility.Build.cs
@@ -33,6 +33,8 @@ public class Blutility : ModuleRules
 				"CollectionManager",
                 "PropertyEditor",
                 "BlueprintGraph",
+                "Json",
+                "JsonUtilities",
 				"UMG",
                 "UMGEditor",
                 "KismetCompiler",
diff --git a/Engine/Source/Editor/Blutility/Classes/ActorActionUtility.h b/Engine/Source/Editor/Blutility/Classes/ActorActionUtility.h
--- a/Engine/Source/Editor/Blutility/Classes/ActorActionUtility.h
+++ b/Engine/Source/Editor/Blutility/Classes/ActorActionUtility.h
@@ -3,10 +3,8 @@
 #pragma once
 
 #include "CoreMinimal.h"
-#include "UObject/ObjectMacros.h"
 #include "EditorUtilityObject.h"
-#include "UObject/ScriptMacros.h"
-#include "BlutilityMenuExtensions.h"
+#include "IEditorUtilityExtension.h"
 #include "ActorActionUtility.generated.h"
 
 /** 
@@ -21,6 +19,27 @@ class BLUTILITY_API UActorActionUtility : public UEditorUtilityObject, public IE
 
 public:
 	/** Return the class that this actor action supports. Leave this blank to support all actor classes. */
-	UFUNCTION(BlueprintPure, BlueprintImplementableEvent, Category="Assets")
+	UE_DEPRECATED(5.2, "GetSupportedClasses() instead, but ideally you're not requesting this directly and are instead using GetSupportedClasses_FromAssetTags.")
+	UFUNCTION(BlueprintPure, BlueprintImplementableEvent, Category="Assets", meta=(DeprecatedFunction, DeprecationMessage="If you were just returning a single class.  Add it to the SupportedClasses array.  If you were doing complex logic to simulate having multiple classes act as filters, add them to the SupportedClasses array.  If you were doing 'other' logic, you'll need to do that upon action execution."))
 	UClass* GetSupportedClass() const;
+	
+	/**
+     * Gets the statically determined supported classes, these classes are used as a first pass filter when determining
+     * if we can utilize this asset utility action on the asset.
+     */
+    UFUNCTION(BlueprintPure, Category = "Assets")
+    const TArray<TSoftClassPtr<UObject>>& GetSupportedClasses() const { return SupportedClasses; }
+
+public:
+    // Begin UObject
+    virtual void GetAssetRegistryTags(TArray<FAssetRegistryTag>& OutTags) const override;
+    // End UObject
+	
+private:
+	/**
+	 * For simple Asset Action's you should fill out the supported class here.  Don't bother with GetSupportedClass()
+	 * * unless you actually need to do specialized dynamic logic.
+	 */
+	UPROPERTY(EditDefaultsOnly, Category="Assets", meta=(AllowAbstract))
+	TArray<TSoftClassPtr<UObject>> SupportedClasses;
 };
\ No newline at end of file
diff --git a/Engine/Source/Editor/Blutility/Classes/AssetActionUtility.h b/Engine/Source/Editor/Blutility/Classes/AssetActionUtility.h
--- a/Engine/Source/Editor/Blutility/Classes/AssetActionUtility.h
+++ b/Engine/Source/Editor/Blutility/Classes/AssetActionUtility.h
@@ -2,13 +2,22 @@
 
 #pragma once
 
-#include "CoreMinimal.h"
 #include "UObject/ObjectMacros.h"
 #include "EditorUtilityObject.h"
 #include "UObject/ScriptMacros.h"
-#include "BlutilityMenuExtensions.h"
+#include "IEditorUtilityExtension.h"
 #include "AssetActionUtility.generated.h"
 
+namespace AssetActionUtilityTags
+{
+	extern const FName BlutilityTagVersion;
+	extern const FName SupportedClasses;
+	extern const FName IsActionForBlueprint;
+	extern const FName CallableFunctions;
+}
+
+struct FAssetData;
+
 /** 
  * Base class for all asset action-related utilities
  * Any functions/events that are exposed on derived classes that have the correct signature will be
@@ -20,14 +29,38 @@ class BLUTILITY_API UAssetActionUtility : public UEditorUtilityObject, public IE
 	GENERATED_BODY()
 
 public:
-	/** Return the class that this asset action supports (if not implemented, it will show up for all asset types) */
-	UFUNCTION(BlueprintPure, BlueprintImplementableEvent, Category="Assets")
+	/**
+	 * Return the class that this asset action supports (if not implemented, it will show up for all asset types)
+	 * Do not do custom logic here based on the currently selected assets.
+	 */
+	UE_DEPRECATED(5.2, "GetSupportedClasses() instead, but ideally you're not requesting this directly and are instead using GetSupportedClasses_FromAssetTags.")
+	UFUNCTION(BlueprintPure, BlueprintImplementableEvent, Category="Assets", meta=(DeprecatedFunction, DeprecationMessage="If you were just returning a single class.  Add it to the SupportedClasses array.  If you were doing complex logic to simulate having multiple classes act as filters, add them to the SupportedClasses array.  If you were doing 'other' logic, you'll need to do that upon action execution."))
 	UClass* GetSupportedClass() const;
 
 	/**
 	 * Returns whether or not this action is designed to work specifically on Blueprints (true) or on all assets (false).
 	 * If true, GetSupportedClass() is treated as a filter against the Parent Class of selected Blueprint assets
 	 */
-	UFUNCTION(BlueprintPure, BlueprintImplementableEvent, Category = "Assets")
+	UFUNCTION(BlueprintPure, BlueprintImplementableEvent, Category="Assets")
 	bool IsActionForBlueprints() const;
+
+	/**
+	 * Gets the statically determined supported classes, these classes are used as a first pass filter when determining
+	 * if we can utilize this asset utility action on the asset.
+	 */
+	UFUNCTION(BlueprintPure, Category = "Assets")
+	const TArray<TSoftClassPtr<UObject>>& GetSupportedClasses() const { return SupportedClasses; }
+
+public:
+	// Begin UObject
+	virtual void GetAssetRegistryTags(TArray<FAssetRegistryTag>& OutTags) const override;
+	// End UObject
+
+protected:
+	/**
+	 * For simple Asset Action's you should fill out the supported class here.  Don't bother with GetSupportedClass()
+	 * * unless you actually need to do specialized dynamic logic.
+	 */
+	UPROPERTY(EditDefaultsOnly, Category="Assets", meta=(AllowAbstract))
+	TArray<TSoftClassPtr<UObject>> SupportedClasses;
 };
\ No newline at end of file
diff --git a/Engine/Source/Editor/Blutility/Private/ActorActionUtility.cpp b/Engine/Source/Editor/Blutility/Private/ActorActionUtility.cpp
new file mode 100644
--- /dev/null
+++ b/Engine/Source/Editor/Blutility/Private/ActorActionUtility.cpp
@@ -0,0 +1,26 @@
+// Copyright Epic Games, Inc. All Rights Reserved.
+
+#include "ActorActionUtility.h"
+
+#include "EditorUtilityAssetPrototype.h"
+#include "EditorUtilityBlueprint.h"
+
+#define LOCTEXT_NAMESPACE "ActorActionUtility"
+
+void UActorActionUtility::GetAssetRegistryTags(TArray<FAssetRegistryTag>& OutTags) const
+{
+	Super::GetAssetRegistryTags(OutTags);
+	
+	const bool IsUpToDate =
+		!GetClass()->IsFunctionImplementedInScript("GetSupportedClass");
+
+	if (IsUpToDate)
+	{
+		FAssetActionUtilityPrototype::AddTagsFor_Version(OutTags);
+	}
+	
+	FAssetActionUtilityPrototype::AddTagsFor_SupportedClasses(SupportedClasses, OutTags);
+	FAssetActionUtilityPrototype::AddTagsFor_CallableFunctions(this, OutTags);
+}
+
+#undef LOCTEXT_NAMESPACE
\ No newline at end of file
diff --git a/Engine/Source/Editor/Blutility/Private/AssetActionUtility.cpp b/Engine/Source/Editor/Blutility/Private/AssetActionUtility.cpp
new file mode 100644
--- /dev/null
+++ b/Engine/Source/Editor/Blutility/Private/AssetActionUtility.cpp
@@ -0,0 +1,31 @@
+// Copyright Epic Games, Inc. All Rights Reserved.
+
+#include "AssetActionUtility.h"
+
+#include "EditorUtilityAssetPrototype.h"
+#include "EditorUtilityBlueprint.h"
+#include "AssetRegistry/AssetData.h"
+#include "JsonObjectConverter.h"
+#include "Misc/DataValidation.h"
+#include "Serialization/JsonSerializerMacros.h"
+
+#define LOCTEXT_NAMESPACE "AssetActionUtility"
+
+void UAssetActionUtility::GetAssetRegistryTags(TArray<FAssetRegistryTag>& OutTags) const
+{
+	Super::GetAssetRegistryTags(OutTags);
+	
+	const bool IsUpToDate =
+		!GetClass()->IsFunctionImplementedInScript("GetSupportedClass");
+	
+	if (IsUpToDate)
+	{
+		FAssetActionUtilityPrototype::AddTagsFor_Version(OutTags);
+	}
+	
+	FAssetActionUtilityPrototype::AddTagsFor_SupportedClasses(SupportedClasses, OutTags);
+	FAssetActionUtilityPrototype::AddTagsFor_IsActionForBlueprints(IsActionForBlueprints(), OutTags);
+	FAssetActionUtilityPrototype::AddTagsFor_CallableFunctions(this, OutTags);
+}
+			
+#undef LOCTEXT_NAMESPACE
\ No newline at end of file
diff --git a/Engine/Source/Editor/Blutility/Private/BlutilityContentBrowserExtensions.cpp b/Engine/Source/Editor/Blutility/Private/BlutilityContentBrowserExtensions.cpp
--- a/Engine/Source/Editor/Blutility/Private/BlutilityContentBrowserExtensions.cpp
+++ b/Engine/Source/Editor/Blutility/Private/BlutilityContentBrowserExtensions.cpp
@@ -10,9 +10,12 @@
 #include "Containers/Set.h"
 #include "ContentBrowserDelegates.h"
 #include "ContentBrowserModule.h"
+#include "EditorUtilityAssetPrototype.h"
 #include "Delegates/Delegate.h"
 #include "EditorUtilityBlueprint.h"
+#include "Algo/AnyOf.h"
 #include "Engine/Blueprint.h"
+#include "Engine/BlueprintGeneratedClass.h"
 #include "Framework/MultiBox/MultiBoxExtender.h"
 #include "HAL/Platform.h"
 #include "HAL/PlatformCrt.h"
@@ -25,6 +28,7 @@
 #include "UObject/Class.h"
 #include "UObject/NameTypes.h"
 #include "UObject/UObjectIterator.h"
+#include "Logging/MessageLog.h"
 
 #define LOCTEXT_NAMESPACE "BlutilityContentBrowserExtensions"
 
@@ -39,111 +43,78 @@ public:
 		TSharedRef<FExtender> Extender(new FExtender());
 
 		// Run thru the assets to determine if any meet our criteria
-		TMap<IEditorUtilityExtension*, TSet<int32>> UtilityAndSelectionIndices;
-		TArray<IEditorUtilityExtension*> SupportedUtils;
+		TMap<TSharedRef<FAssetActionUtilityPrototype>, TSet<int32>> UtilityAndSelectionIndices;
 		TArray<FAssetData> SupportedAssets;
+		
 		if (SelectedAssets.Num() > 0)
 		{
-			// Check blueprint utils (we need to load them to query their validity against these assets)
-			TArray<FAssetData> UtilAssets;
-			FBlutilityMenuExtensions::GetBlutilityClasses(UtilAssets, UAssetActionUtility::StaticClass()->GetClassPathName());
-
-			// Collect all UAssetActionUtility derived classes
-			TSet<UAssetActionUtility*> AssetClasses;
-			for (TObjectIterator<UClass> AssetActionClassIt; AssetActionClassIt; ++AssetActionClassIt)
-			{
-				if (AssetActionClassIt->IsChildOf(UAssetActionUtility::StaticClass()) && UAssetActionUtility::StaticClass()->GetFName() != AssetActionClassIt->GetFName() && AssetActionClassIt->ClassGeneratedBy == nullptr)
-				{
-					AssetClasses.Add(Cast<UAssetActionUtility>(AssetActionClassIt->GetDefaultObject()));
-				}
-			}
-
-			auto ProcessAssetAction = [&SupportedAssets, &UtilityAndSelectionIndices, &SelectedAssets](UAssetActionUtility* InAction)
+			FMessageLog EditorErrors("EditorErrors");
+			EditorErrors.NewPage(LOCTEXT("ScriptedActions", "Scripted Actions"));
+			
+			auto ProcessAssetAction = [&EditorErrors, &SupportedAssets, &UtilityAndSelectionIndices, &SelectedAssets](const TSharedRef<FAssetActionUtilityPrototype>& ActionUtilityPrototype)
 			{
-				if (UClass* SupportedClass = InAction->GetSupportedClass())
+				if (ActionUtilityPrototype->IsLatestVersion())
 				{
-					const bool bIsActionForBlueprints = InAction->IsActionForBlueprints();
-
-					for (const FAssetData& Asset : SelectedAssets)
+					TArray<TSoftClassPtr<UObject>> SupportedClassPtrs = ActionUtilityPrototype->GetSupportedClasses();
+					if (SupportedClassPtrs.Num() > 0)
 					{
-						bool bPassesClassFilter = false;
-						if (bIsActionForBlueprints)
+						const bool bIsActionForBlueprints = ActionUtilityPrototype->AreSupportedClassesForBlueprints();
+
+						for (const FAssetData& Asset : SelectedAssets)
 						{
-							if (TSubclassOf<UBlueprint> AssetClass = Asset.GetClass())
+							bool bPassesClassFilter = false;
+							if (bIsActionForBlueprints)
 							{
-								if (UBlueprint* AssetAsBlueprint = Cast<UBlueprint>(Asset.GetAsset()))
-								{
-									// It's a blueprint, but is it the right kind?
-									bPassesClassFilter = AssetAsBlueprint->ParentClass && AssetAsBlueprint->ParentClass->IsChildOf(SupportedClass);
-								}
-								else
+								if (TSubclassOf<UBlueprint> AssetClass = Asset.GetClass())
 								{
-									// Not a blueprint
-									bPassesClassFilter = false;
+									if (const UClass* Blueprint_ParentClass = UBlueprint::GetBlueprintParentClassFromAssetTags(Asset))
+	                                {
+										bPassesClassFilter = 
+											Algo::AnyOf(SupportedClassPtrs, [Blueprint_ParentClass](TSoftClassPtr<UObject> ClassPtr){ return Blueprint_ParentClass->IsChildOf(ClassPtr.Get()); });
+	                                }
 								}
 							}
 							else
 							{
-								// Not a blueprint
-								bPassesClassFilter = false;
+								// Is the asset the right kind?
+								bPassesClassFilter = 
+									Algo::AnyOf(SupportedClassPtrs, [&Asset](TSoftClassPtr<UObject> ClassPtr){ return Asset.IsInstanceOf(ClassPtr.Get()); });
 							}
-						}
-						else
-						{
-							// Is the asset the right kind?
-							bPassesClassFilter = Asset.IsInstanceOf(SupportedClass);
-						}
 
-						if (bPassesClassFilter)
-						{
-							const int32 Index = SupportedAssets.AddUnique(Asset);
-							UtilityAndSelectionIndices.FindOrAdd(InAction).Add(Index);
+							if (bPassesClassFilter)
+							{
+								const int32 Index = SupportedAssets.AddUnique(Asset);
+								UtilityAndSelectionIndices.FindOrAdd(ActionUtilityPrototype).Add(Index);
+							}
 						}
 					}
 				}
 				else
 				{
-					TSet<int32>& ActionIndices = UtilityAndSelectionIndices.FindOrAdd(InAction);
-					for (const FAssetData& Asset : SelectedAssets)
-					{ 
-						const int32 Index = SupportedAssets.AddUnique(Asset);
-						ActionIndices.Add(Index);
-					}
+					TSharedRef<FTokenizedMessage> ErrorMessage = EditorErrors.Error();
+                    ErrorMessage->AddToken(FAssetNameToken::Create(ActionUtilityPrototype->GetUtilityBlueprintAsset().GetObjectPathString(),FText::FromString(ActionUtilityPrototype->GetUtilityBlueprintAsset().GetObjectPathString())));
+                    ErrorMessage->AddToken(FTextToken::Create(LOCTEXT("NeedsToBeUpdated", "needs to be re-saved and possibly upgraded.")));
 				}
 			};
 
+			// Check blueprint utils (we need to load them to query their validity against these assets)
+			TArray<FAssetData> UtilAssets;
+			FBlutilityMenuExtensions::GetBlutilityClasses(UtilAssets, UAssetActionUtility::StaticClass()->GetClassPathName());
+
 			// Process asset based utilities
 			for (const FAssetData& UtilAsset : UtilAssets)
 			{
-				FString ParentClassName;
-				if (UtilAsset.GetTagValue(FBlueprintTags::NativeParentClassPath, ParentClassName))
+				if (const UClass* ParentClass = UBlueprint::GetBlueprintParentClassFromAssetTags(UtilAsset))
 				{
-					UObject* Outer = nullptr;
-					ResolveName(Outer, ParentClassName, false, false);
-					const UClass* ParentClass = FindObject<UClass>(Outer, *ParentClassName);
-
 					// We only care about UEditorUtilityBlueprint's that are compiling subclasses of UAssetActionUtility
-					if (ParentClass && ParentClass->IsChildOf(UAssetActionUtility::StaticClass()))
+					if (ParentClass->IsChildOf(UAssetActionUtility::StaticClass()))
 					{
-						if (const UEditorUtilityBlueprint* Blueprint = Cast<UEditorUtilityBlueprint>(UtilAsset.GetAsset()))
-						{
-							if (const UClass* BPClass = Blueprint->GeneratedClass.Get())
-							{
-								if (UAssetActionUtility* DefaultObject = Cast<UAssetActionUtility>(BPClass->GetDefaultObject()))
-								{
-									ProcessAssetAction(DefaultObject);
-								}
-							}
-						}
+						ProcessAssetAction(MakeShared<FAssetActionUtilityPrototype>(UtilAsset));
 					}
 				}
 			}
 
-			// Process non-asset based utilities
-			for (UAssetActionUtility* Action : AssetClasses)
-			{
-				ProcessAssetAction(Action);
-			}
+			EditorErrors.Notify(LOCTEXT("SomeProblemsWithAssetActionUtility", "There were some problems with some AssetActionUtility Blueprints."));
 		}
 
 		if (UtilityAndSelectionIndices.Num() > 0)
diff --git a/Engine/Source/Editor/Blutility/Private/BlutilityLevelEditorExtensions.cpp b/Engine/Source/Editor/Blutility/Private/BlutilityLevelEditorExtensions.cpp
--- a/Engine/Source/Editor/Blutility/Private/BlutilityLevelEditorExtensions.cpp
+++ b/Engine/Source/Editor/Blutility/Private/BlutilityLevelEditorExtensions.cpp
@@ -9,12 +9,15 @@
 #include "Containers/Map.h"
 #include "Containers/Set.h"
 #include "CoreTypes.h"
+#include "EditorUtilityAssetPrototype.h"
 #include "Delegates/Delegate.h"
 #include "EditorUtilityBlueprint.h"
 #include "Framework/MultiBox/MultiBoxExtender.h"
 #include "GameFramework/Actor.h"
 #include "HAL/PlatformCrt.h"
 #include "LevelEditor.h"
+#include "Algo/AnyOf.h"
+#include "Logging/MessageLog.h"
 #include "Modules/ModuleManager.h"
 #include "Templates/Casts.h"
 #include "Templates/SharedPointer.h"
@@ -36,76 +39,65 @@ public:
 	{
 		TSharedRef<FExtender> Extender(new FExtender());
 			   
-		TMap<IEditorUtilityExtension*, TSet<int32>> UtilityAndSelectionIndices;
+		TMap<TSharedRef<FAssetActionUtilityPrototype>, TSet<int32>> UtilityAndSelectionIndices;
 
 		// Run thru the actors to determine if any meet our criteria
 		TArray<AActor*> SupportedActors;
 		if (SelectedActors.Num() > 0)
 		{
-			// Check blueprint utils (we need to load them to query their validity against these assets)
-			TArray<FAssetData> UtilAssets;
-			FBlutilityMenuExtensions::GetBlutilityClasses(UtilAssets, UActorActionUtility::StaticClass()->GetClassPathName());
-
-			TMap<UActorActionUtility*, UClass*> ActorActionUtilities;
-
-			// Process asset based utilities
-			for (const FAssetData& UtilAsset : UtilAssets)
+			FMessageLog EditorErrors("EditorErrors");
+			EditorErrors.NewPage(LOCTEXT("ScriptedActions", "Scripted Actions"));
+			
+			auto ProcessActorAction = [&UtilityAndSelectionIndices, &SelectedActors, &SupportedActors, &EditorErrors](const TSharedRef<FAssetActionUtilityPrototype>& ActionUtilityPrototype)
 			{
-				if (UEditorUtilityBlueprint* Blueprint = Cast<UEditorUtilityBlueprint>(UtilAsset.GetAsset()))
+				if (ActionUtilityPrototype->IsLatestVersion())
 				{
-					if (UClass* BPClass = Blueprint->GeneratedClass.Get())
+					TArray<TSoftClassPtr<UObject>> SupportedClassPtrs = ActionUtilityPrototype->GetSupportedClasses();
+					if (SupportedClassPtrs.Num() > 0)
 					{
-						if (UActorActionUtility* ActorActionUtility = Cast<UActorActionUtility>(BPClass->GetDefaultObject()))
+						for (AActor* Actor : SelectedActors)
 						{
-							UClass*& SupportedClass = ActorActionUtilities.FindOrAdd(ActorActionUtility);
-							if (!SupportedClass)
+							if (Actor)
 							{
-								SupportedClass = ActorActionUtility->GetSupportedClass();
+								UClass* ActorClass = Actor->GetClass();
+								const bool bPassesClassFilter = 
+									Algo::AnyOf(SupportedClassPtrs, [ActorClass](TSoftClassPtr<UObject> ClassPtr){ return ActorClass->IsChildOf(ClassPtr.Get()); });
+
+								if (bPassesClassFilter)
+								{
+									const int32 Index = SupportedActors.AddUnique(Actor);
+									UtilityAndSelectionIndices.FindOrAdd(ActionUtilityPrototype).Add(Index);
+								}
 							}
 						}
 					}
 				}
-			}
+				else
+                {
+					TSharedRef<FTokenizedMessage> ErrorMessage = EditorErrors.Error();
+                    ErrorMessage->AddToken(FAssetNameToken::Create(ActionUtilityPrototype->GetUtilityBlueprintAsset().GetObjectPathString(),FText::FromString(ActionUtilityPrototype->GetUtilityBlueprintAsset().GetObjectPathString())));
+                    ErrorMessage->AddToken(FTextToken::Create(LOCTEXT("NeedsToBeUpdated", "needs to be re-saved and possibly upgraded.")));
+                }
+			};
 
-			// Process non-asset based utilities
-			for (TObjectIterator<UClass> AssetActionClassIt; AssetActionClassIt; ++AssetActionClassIt)
-			{
-				if (AssetActionClassIt->IsChildOf(UActorActionUtility::StaticClass()) && UActorActionUtility::StaticClass()->GetFName() != AssetActionClassIt->GetFName() && AssetActionClassIt->ClassGeneratedBy == nullptr)
-				{
-					if (UActorActionUtility* ActorActionUtility = Cast<UActorActionUtility>(AssetActionClassIt->GetDefaultObject()))
-					{
-						UClass*& SupportedClass = ActorActionUtilities.FindOrAdd(ActorActionUtility);
-						if (!SupportedClass)
-						{
-							SupportedClass = ActorActionUtility->GetSupportedClass();
-						}
-					}
-				}
-			}
+			// Check blueprint utils (we need to load them to query their validity against these assets)
+			TArray<FAssetData> UtilAssets;
+			FBlutilityMenuExtensions::GetBlutilityClasses(UtilAssets, UActorActionUtility::StaticClass()->GetClassPathName());
 
-			if (UtilAssets.Num() + ActorActionUtilities.Num() > 0)
-			{
-				for (AActor* Actor : SelectedActors)
-				{
-					if (Actor)
-					{
-						int32 SupportedActorIndex = INDEX_NONE;
-						for (const TPair<UActorActionUtility*, UClass*>& ActorActionUtilityIt : ActorActionUtilities)
-						{
-							UClass* SupportedClass = ActorActionUtilityIt.Value;
-							if (SupportedClass == nullptr || (SupportedClass && Actor->GetClass()->IsChildOf(SupportedClass)))
-							{
-								UActorActionUtility* Action = ActorActionUtilityIt.Key;
-								if (SupportedActorIndex == INDEX_NONE)
-								{
-									SupportedActorIndex = SupportedActors.Add(Actor);
-								}
-								UtilityAndSelectionIndices.FindOrAdd(Action).Add(SupportedActorIndex);
-							}
-						}
+			// Process asset based utilities
+            for (const FAssetData& UtilAsset : UtilAssets)
+            {
+            	if (const UClass* ParentClass = UBlueprint::GetBlueprintParentClassFromAssetTags(UtilAsset))
+            	{
+            		// We only care about UEditorUtilityBlueprint's that are compiling subclasses of UActorActionUtility
+            		if (ParentClass->IsChildOf(UActorActionUtility::StaticClass()))
+            		{
+            			ProcessActorAction(MakeShared<FAssetActionUtilityPrototype>(UtilAsset));
 					}
 				}
 			}
+			
+			EditorErrors.Notify(LOCTEXT("SomeProblemsWithActorActionUtility", "There were some problems with some ActorActionUtility Blueprints."));
 		}
 
 		if (SupportedActors.Num() > 0)
diff --git a/Engine/Source/Editor/Blutility/Private/BlutilityMenuExtensions.cpp b/Engine/Source/Editor/Blutility/Private/BlutilityMenuExtensions.cpp
--- a/Engine/Source/Editor/Blutility/Private/BlutilityMenuExtensions.cpp
+++ b/Engine/Source/Editor/Blutility/Private/BlutilityMenuExtensions.cpp
@@ -14,6 +14,7 @@
 #include "DetailsViewArgs.h"
 #include "EdGraphSchema_K2.h"
 #include "Editor.h"
+#include "EditorUtilityAssetPrototype.h"
 #include "Editor/EditorEngine.h"
 #include "EditorUtilityBlueprint.h"
 #include "Engine/Blueprint.h"
@@ -23,6 +24,7 @@
 #include "GameFramework/Actor.h"
 #include "GenericPlatform/GenericApplication.h"
 #include "IDetailsView.h"
+#include "IEditorUtilityExtension.h"
 #include "IStructureDetailsView.h"
 #include "Input/Reply.h"
 #include "Internationalization/Internationalization.h"
@@ -222,10 +224,9 @@ void FBlutilityMenuExtensions::GetBlutilityClasses(TArray<FAssetData>& OutAssets
 	// Get class names
 	TArray<FTopLevelAssetPath> BaseNames;
 	BaseNames.Add(InClassName);
-	TSet<FTopLevelAssetPath> Excluded;
-	TSet<FTopLevelAssetPath> DerivedNames;
-	AssetRegistry.GetDerivedClassNames(BaseNames, Excluded, DerivedNames);
-
+	TSet<FTopLevelAssetPath> DerivedClasses;
+	AssetRegistry.GetDerivedClassNames(BaseNames, TSet<FTopLevelAssetPath>(), DerivedClasses);
+	
 	// Now get all UEditorUtilityBlueprint assets
 	FARFilter Filter;
 	Filter.ClassPaths.Add(UEditorUtilityBlueprint::StaticClass()->GetClassPathName());
@@ -238,19 +239,35 @@ void FBlutilityMenuExtensions::GetBlutilityClasses(TArray<FAssetData>& OutAssets
 	// Check each asset to see if it matches our type
 	for (const FAssetData& Asset : AssetList)
 	{
+		// Abstract or deprecated blutilities should not be included.
+		const EClassFlags BPFlags = static_cast<EClassFlags>(Asset.GetTagValueRef<uint32>(FBlueprintTags::ClassFlags));
+		if (EnumHasAnyFlags(BPFlags, CLASS_Abstract | CLASS_Deprecated))
+		{
+			continue;
+		}
+		
 		FAssetDataTagMapSharedView::FFindTagResult Result = Asset.TagsAndValues.FindTag(FBlueprintTags::GeneratedClassPath);
 		if (Result.IsSet())
 		{
+			// If it's a menu extension blutility, don't include ones from other peoples developer folders, just your own.
+			if (FPaths::IsUnderDirectory(Asset.PackagePath.ToString(), FPaths::GameDevelopersDir()))
+			{
+				if (!FPaths::IsUnderDirectory(Asset.PackagePath.ToString(), FPaths::GameUserDeveloperFolderName()))
+				{
+					continue;
+				}
+			}
+			
 			const FTopLevelAssetPath ClassObjectPath(FPackageName::ExportTextPathToObjectPath(Result.GetValue()));
 
-			if (DerivedNames.Contains(ClassObjectPath))
+			if (DerivedClasses.Contains(ClassObjectPath))
 			{
 				OutAssets.Add(Asset);
 			}
 		}
 	}
 }
-void FBlutilityMenuExtensions::CreateActorBlutilityActionsMenu(FMenuBuilder& MenuBuilder, TMap<class IEditorUtilityExtension*, TSet<int32>> Utils, const TArray<AActor*> SelectedSupportedActors)
+void FBlutilityMenuExtensions::CreateActorBlutilityActionsMenu(FMenuBuilder& MenuBuilder, TMap<TSharedRef<FAssetActionUtilityPrototype>, TSet<int32>> Utils, const TArray<AActor*> SelectedSupportedActors)
 {
 	CreateBlutilityActionsMenu<AActor*>(MenuBuilder, Utils,
 		LOCTEXT("ScriptedActorActions", "Scripted Actor Actions"),
@@ -269,7 +286,7 @@ void FBlutilityMenuExtensions::CreateActorBlutilityActionsMenu(FMenuBuilder& Men
 	);
 }
 
-void FBlutilityMenuExtensions::CreateAssetBlutilityActionsMenu(FMenuBuilder& MenuBuilder, TMap<class IEditorUtilityExtension*, TSet<int32>> Utils, const TArray<FAssetData> SelectedSupportedAssets)
+void FBlutilityMenuExtensions::CreateAssetBlutilityActionsMenu(FMenuBuilder& MenuBuilder, TMap<TSharedRef<FAssetActionUtilityPrototype>, TSet<int32>> Utils, const TArray<FAssetData> SelectedSupportedAssets)
 {
 	CreateBlutilityActionsMenu<FAssetData>(MenuBuilder, Utils,
 		LOCTEXT("ScriptedAssetActions", "Scripted Asset Actions"),
@@ -293,70 +310,50 @@ void FBlutilityMenuExtensions::CreateAssetBlutilityActionsMenu(FMenuBuilder& Men
 void FBlutilityMenuExtensions::OpenEditorForUtility(const FFunctionAndUtil& FunctionAndUtil)
 {
 	// Edit the script if we have shift held down
-	if (UBlueprint* Blueprint = Cast<UBlueprint>(Cast<UObject>(FunctionAndUtil.Util)->GetClass()->ClassGeneratedBy))
+	if (UBlueprint* Blueprint = Cast<UBlueprint>(Cast<UObject>(FunctionAndUtil.Util->LoadUtilityAsset())->GetClass()->ClassGeneratedBy))
 	{
 		if (IAssetEditorInstance* AssetEditor = GEditor->GetEditorSubsystem<UAssetEditorSubsystem>()->FindEditorForAsset(Blueprint, true))
 		{
 			check(AssetEditor->GetEditorName() == TEXT("BlueprintEditor"));
 			IBlueprintEditor* BlueprintEditor = static_cast<IBlueprintEditor*>(AssetEditor);
-			BlueprintEditor->JumpToHyperlink(FunctionAndUtil.Function, false);
+			BlueprintEditor->JumpToHyperlink(FunctionAndUtil.GetFunction(), false);
 		}
 		else
 		{
 			FBlueprintEditorModule& BlueprintEditorModule = FModuleManager::LoadModuleChecked<FBlueprintEditorModule>("Kismet");
 			TSharedRef<IBlueprintEditor> BlueprintEditor = BlueprintEditorModule.CreateBlueprintEditor(EToolkitMode::Standalone, TSharedPtr<IToolkitHost>(), Blueprint, false);
-			BlueprintEditor->JumpToHyperlink(FunctionAndUtil.Function, false);
+			BlueprintEditor->JumpToHyperlink(FunctionAndUtil.GetFunction(), false);
 		}
 	}
 }
 
-void FBlutilityMenuExtensions::ExtractFunctions(TMap<class IEditorUtilityExtension*, TSet<int32>>& Utils, TMap<FString, TArray<FFunctionAndUtil>>& OutCategoryFunctions)
+void FBlutilityMenuExtensions::ExtractFunctions(TMap<TSharedRef<FAssetActionUtilityPrototype>, TSet<int32>>& Utils, TMap<FString, TArray<FFunctionAndUtil>>& OutCategoryFunctions)
 {
 	TSet<UClass*> ProcessedClasses;
 	const static FName NAME_CallInEditor(TEXT("CallInEditor"));
 
 	// Find the exposed functions available in each class, making sure to not list shared functions from a parent class more than once
-	for (TPair< IEditorUtilityExtension*, TSet<int32>> UtilitySelectionPair : Utils)
+	for (TPair<TSharedRef<FAssetActionUtilityPrototype>, TSet<int32>> UtilitySelectionPair : Utils)
 	{
-		IEditorUtilityExtension* Util = UtilitySelectionPair.Key;
-		UClass* Class = Cast<UObject>(Util)->GetClass();
-
-		if (ProcessedClasses.Contains(Class))
-		{
-			continue;
-		}
-
-		for (UClass* ParentClass = Class; ParentClass != UObject::StaticClass(); ParentClass = ParentClass->GetSuperClass())
-		{
-			ProcessedClasses.Add(ParentClass);
-		}
+		const TSharedRef<FAssetActionUtilityPrototype>& Util = UtilitySelectionPair.Key;
 
-		for (TFieldIterator<UFunction> FunctionIt(Class); FunctionIt; ++FunctionIt)
+		TArray<FBlutilityFunctionData> FunctionDatas = Util->GetCallableFunctions();
+		for (const FBlutilityFunctionData& FunctionData : FunctionDatas)
 		{
-			if (UFunction* Func = *FunctionIt)
-			{
-				if (Func->HasMetaData(NAME_CallInEditor) && Func->GetReturnProperty() == nullptr)
-				{
-					const static FName NAME_Category(TEXT("Category"));
-
-					const FString& FunctionCategory = Func->GetMetaData(NAME_Category);
-					TArray<FFunctionAndUtil>& Functions = OutCategoryFunctions.FindOrAdd(FunctionCategory);
-
-					Functions.AddUnique(FFunctionAndUtil(Func, Util, UtilitySelectionPair.Value));
-				}
-			}
+			TArray<FFunctionAndUtil>& Functions = OutCategoryFunctions.FindOrAdd(FunctionData.Category);
+			Functions.AddUnique(FFunctionAndUtil(FunctionData, Util, UtilitySelectionPair.Value));
 		}
 	}
 
 	for (TPair<FString, TArray<FFunctionAndUtil>>& CategoryFunctionPair : OutCategoryFunctions)
 	{
 		// Sort the functions by name
-		CategoryFunctionPair.Value.Sort([](const FFunctionAndUtil& A, const FFunctionAndUtil& B) { return A.Function->GetName() < B.Function->GetName(); });
+		CategoryFunctionPair.Value.Sort([](const FFunctionAndUtil& A, const FFunctionAndUtil& B) { return A.FunctionData.Name.LexicalLess(B.FunctionData.Name); });
 	}
 }
 
 template<typename SelectionType>
-void FBlutilityMenuExtensions::CreateBlutilityActionsMenu(FMenuBuilder& MenuBuilder, TMap<IEditorUtilityExtension*, TSet<int32>> Utils, const FText& MenuLabel, const FText& MenuToolTip, TFunction<bool(const FProperty * Property)> IsValidPropertyType, const TArray<SelectionType> Selection, const FName& IconName)
+void FBlutilityMenuExtensions::CreateBlutilityActionsMenu(FMenuBuilder& MenuBuilder, TMap<TSharedRef<FAssetActionUtilityPrototype>, TSet<int32>> Utils, const FText& MenuLabel, const FText& MenuToolTip, TFunction<bool(const FProperty * Property)> IsValidPropertyType, const TArray<SelectionType> Selection, const FName& IconName)
 {
 	TMap<FString, TArray<FFunctionAndUtil>> CategoryFunctions;
 	ExtractFunctions(Utils, CategoryFunctions);
@@ -365,10 +362,10 @@ void FBlutilityMenuExtensions::CreateBlutilityActionsMenu(FMenuBuilder& MenuBuil
 	{
 		for (const FFunctionAndUtil& FunctionAndUtil : FunctionUtils)
 		{
-			const FText TooltipText = FText::Format(LOCTEXT("AssetUtilTooltipFormat", "{0}\n(Shift-click to edit script)"), FunctionAndUtil.Function->GetToolTipText());
+			const FText TooltipText = FText::Format(LOCTEXT("AssetUtilTooltipFormat", "{0}\n(Shift-click to edit script)"), FunctionAndUtil.FunctionData.TooltipText);
 
 			SubMenuBuilder.AddMenuEntry(
-				FunctionAndUtil.Function->GetDisplayNameText(),
+				FunctionAndUtil.FunctionData.NameText,
 				TooltipText,
 				FSlateIcon(FAppStyle::GetAppStyleSetName(), "GraphEditor.Event_16x"),
 				FExecuteAction::CreateLambda([FunctionAndUtil, Selection, IsValidPropertyType]
@@ -380,21 +377,23 @@ void FBlutilityMenuExtensions::CreateBlutilityActionsMenu(FMenuBuilder& MenuBuil
 					else
 					{
 						// We dont run this on the CDO, as bad things could occur!
-						UObject* TempObject = NewObject<UObject>(GetTransientPackage(), Cast<UObject>(FunctionAndUtil.Util)->GetClass());
+						UObject* TempObject = NewObject<UObject>(GetTransientPackage(), Cast<UObject>(FunctionAndUtil.Util->LoadUtilityAsset())->GetClass());
 						TempObject->AddToRoot(); // Some Blutility actions might run GC so the TempObject needs to be rooted to avoid getting destroyed
 
-						if (FunctionAndUtil.Function->NumParms > 0)
+						UFunction* Function = FunctionAndUtil.GetFunction();
+						
+						if (Function->NumParms > 0)
 						{
 							// Create a parameter struct and fill in defaults
-							TSharedRef<FStructOnScope> FuncParams = MakeShared<FStructOnScope>(FunctionAndUtil.Function);
+							TSharedRef<FStructOnScope> FuncParams = MakeShared<FStructOnScope>(Function);
 
 							FProperty* FirstParamProperty = nullptr;
 
 							int32 ParameterIndex = 0;
-							for (TFieldIterator<FProperty> It(FunctionAndUtil.Function); It&& It->HasAnyPropertyFlags(CPF_Parm); ++It)
+							for (TFieldIterator<FProperty> It(Function); It&& It->HasAnyPropertyFlags(CPF_Parm); ++It)
 							{
 								FString Defaults;
-								if (UEdGraphSchema_K2::FindFunctionParameterDefaultValue(FunctionAndUtil.Function, *It, Defaults))
+								if (UEdGraphSchema_K2::FindFunctionParameterDefaultValue(Function, *It, Defaults))
 								{
 									It->ImportText_Direct(*Defaults, It->ContainerPtrToValuePtr<uint8>(FuncParams->GetStructMemory()), nullptr, PPF_None);
 								}
@@ -414,7 +413,7 @@ void FBlutilityMenuExtensions::CreateBlutilityActionsMenu(FMenuBuilder& MenuBuil
 							{
 								// pop up a dialog to input params to the function
 								TSharedRef<SWindow> Window = SNew(SWindow)
-									.Title(FunctionAndUtil.Function->GetDisplayNameText())
+									.Title(Function->GetDisplayNameText())
 									.ClientSize(FVector2D(400, 200))
 									.SupportsMinimize(false)
 									.SupportsMaximize(false);
@@ -423,7 +422,7 @@ void FBlutilityMenuExtensions::CreateBlutilityActionsMenu(FMenuBuilder& MenuBuil
 								Window->SetContent(
 									SAssignNew(Dialog, SFunctionParamDialog, Window, FuncParams, FirstParamProperty ? FirstParamProperty->GetFName() : NAME_None)
 									.OkButtonText(LOCTEXT("OKButton", "OK"))
-									.OkButtonTooltipText(FunctionAndUtil.Function->GetToolTipText()));
+									.OkButtonTooltipText(Function->GetToolTipText()));
 
 								GEditor->EditorAddModalWindow(Window);
 								bApply = Dialog->bOKPressed;
@@ -438,20 +437,20 @@ void FBlutilityMenuExtensions::CreateBlutilityActionsMenu(FMenuBuilder& MenuBuil
 								if (bForwardUserSelection)
 								{
 									// For each user-select asset forward the selection object into the function first's parameter (if it matches)
-									const FString Path = FirstParamProperty->GetPathName(FunctionAndUtil.Function);
+									const FString Path = FirstParamProperty->GetPathName(Function);
 
 									// Ensure we only process selection objects that are valid for this function/utility
 									for (const int32& SelectionIndex : FunctionAndUtil.SelectionIndices)
 									{
 										const auto SelectedAsset = Selection[SelectionIndex];
 										FirstParamProperty->CopySingleValue(FirstParamProperty->ContainerPtrToValuePtr<uint8>(FuncParams->GetStructMemory()), &SelectedAsset);
-										TempObject->ProcessEvent(FunctionAndUtil.Function, FuncParams->GetStructMemory());
+										TempObject->ProcessEvent(Function, FuncParams->GetStructMemory());
 									}
 								}
 								else
 								{
 									// User is expected to manage the asset selection on its own
-									TempObject->ProcessEvent(FunctionAndUtil.Function, FuncParams->GetStructMemory());
+									TempObject->ProcessEvent(Function, FuncParams->GetStructMemory());
 								}
 							}
 						}
@@ -459,7 +458,7 @@ void FBlutilityMenuExtensions::CreateBlutilityActionsMenu(FMenuBuilder& MenuBuil
 						{
 							FScopedTransaction Transaction(NSLOCTEXT("UnrealEd", "BlutilityAction", "Blutility Action"));
 							FEditorScriptExecutionGuard ScriptGuard;
-							TempObject->ProcessEvent(FunctionAndUtil.Function, nullptr);
+							TempObject->ProcessEvent(Function, nullptr);
 						}
 
 						TempObject->RemoveFromRoot();
diff --git a/Engine/Source/Editor/Blutility/Private/BlutilityMenuExtensions.h b/Engine/Source/Editor/Blutility/Private/BlutilityMenuExtensions.h
--- a/Engine/Source/Editor/Blutility/Private/BlutilityMenuExtensions.h
+++ b/Engine/Source/Editor/Blutility/Private/BlutilityMenuExtensions.h
@@ -2,6 +2,7 @@
 
 #pragma once
 
+#include "EditorUtilityAssetPrototype.h"
 #include "Containers/Array.h"
 #include "Containers/Map.h"
 #include "Containers/Set.h"
@@ -14,8 +15,6 @@
 #include "UObject/TopLevelAssetPath.h"
 #include "UObject/UObjectGlobals.h"
 
-#include "BlutilityMenuExtensions.generated.h"
-
 class AActor;
 class FMenuBuilder;
 class FProperty;
@@ -24,7 +23,7 @@ class FText;
 class IEditorUtilityExtension;
 class UFunction;
 class UObject;
-struct FAssetData;
+class FAssetActionUtilityPrototype;
 
 // Blutility Menu extension helpers
 class FBlutilityMenuExtensions
@@ -34,45 +33,47 @@ public:
 	static void GetBlutilityClasses(TArray<FAssetData>& OutAssets, FTopLevelAssetPath InClassName);
 
 	/** Helper function that populates a menu based on the exposed functions in a set of Blutility objects */
-	static void CreateAssetBlutilityActionsMenu(FMenuBuilder& MenuBuilder, TMap<class IEditorUtilityExtension*, TSet<int32>> Utils, const TArray<FAssetData> SelectedSupportedAssets);
-	static void CreateActorBlutilityActionsMenu(FMenuBuilder& MenuBuilder, TMap<class IEditorUtilityExtension*, TSet<int32>> Utils, const TArray<AActor*> SelectedSupportedActors);
+	static void CreateAssetBlutilityActionsMenu(FMenuBuilder& MenuBuilder, TMap<TSharedRef<FAssetActionUtilityPrototype>, TSet<int32>> Utils, const TArray<FAssetData> SelectedSupportedAssets);
+	static void CreateActorBlutilityActionsMenu(FMenuBuilder& MenuBuilder, TMap<TSharedRef<FAssetActionUtilityPrototype>, TSet<int32>> Utils, const TArray<AActor*> SelectedSupportedActors);
 	
 protected:
 	// Helper struct to track the util to call a function on
 	struct FFunctionAndUtil
 	{
-		FFunctionAndUtil(UFunction* InFunction, IEditorUtilityExtension* InUtil, TSet<int32>& InSelection)
-			: Function(InFunction)
+		FFunctionAndUtil(const FBlutilityFunctionData& InFunctionData, const TSharedRef<FAssetActionUtilityPrototype>& InUtil, TSet<int32>& InSelection)
+			: FunctionData(InFunctionData)
 			, Util(InUtil) 
 			, SelectionIndices(InSelection)
 		{}
 
 		bool operator==(const FFunctionAndUtil& InFunction) const
 		{
-			return InFunction.Function == Function && InFunction.Util == Util;
+			return InFunction.FunctionData == InFunction.FunctionData && InFunction.Util == Util;
+		}
+
+		UFunction* GetFunction() const
+		{
+			if (const UObject* UtilityCDO = Util->LoadUtilityAsset())
+			{
+				if (const UClass* UtilityClass = UtilityCDO->GetClass())
+				{
+					return UtilityClass->FindFunctionByName(FunctionData.Name);
+				}
+			}
+
+			return nullptr;
 		}
 
-		UFunction* Function;
-		IEditorUtilityExtension* Util;
+		FBlutilityFunctionData FunctionData;
+		TSharedRef<FAssetActionUtilityPrototype> Util;
 		/** Indices to original object selection array which are supported by this utility */
 		TSet<int32> SelectionIndices;
 	};
 
 protected:
 	template<typename SelectionType>
-	static void CreateBlutilityActionsMenu(FMenuBuilder& MenuBuilder, TMap<class IEditorUtilityExtension*, TSet<int32>> Utils, const FText& MenuLabel, const FText& MenuToolTip, TFunction<bool(const FProperty * Property)> IsValidPropertyType, const TArray<SelectionType> Selection, const FName& IconName = "GraphEditor.Event_16x");
+	static void CreateBlutilityActionsMenu(FMenuBuilder& MenuBuilder, TMap<TSharedRef<FAssetActionUtilityPrototype>, TSet<int32>> Utils, const FText& MenuLabel, const FText& MenuToolTip, TFunction<bool(const FProperty* Property)> IsValidPropertyType, const TArray<SelectionType> Selection, const FName& IconName = "GraphEditor.Event_16x");
 
 	static void OpenEditorForUtility(const FFunctionAndUtil& FunctionAndUtil);
-	static void ExtractFunctions(TMap<class IEditorUtilityExtension*, TSet<int32>>& Utils, TMap<FString, TArray<FFunctionAndUtil>>& OutCategoryFunctions);
+	static void ExtractFunctions(TMap<TSharedRef<FAssetActionUtilityPrototype>, TSet<int32>>& Utils, TMap<FString, TArray<FFunctionAndUtil>>& OutCategoryFunctions);
 };
-
-UINTERFACE(BlueprintType)
-class UEditorUtilityExtension : public UInterface
-{
-	GENERATED_UINTERFACE_BODY()
-};
-
-class IEditorUtilityExtension
-{
-	GENERATED_IINTERFACE_BODY()
-};
\ No newline at end of file
diff --git a/Engine/Source/Editor/Blutility/Private/EditorUtilityAssetPrototype.cpp b/Engine/Source/Editor/Blutility/Private/EditorUtilityAssetPrototype.cpp
new file mode 100644
--- /dev/null
+++ b/Engine/Source/Editor/Blutility/Private/EditorUtilityAssetPrototype.cpp
@@ -0,0 +1,143 @@
+// Copyright Epic Games, Inc. All Rights Reserved.
+
+#include "EditorUtilityAssetPrototype.h"
+#include "EditorUtilityBlueprint.h"
+#include "JsonObjectConverter.h"
+#include "AssetRegistry/AssetData.h"
+#include "Serialization/JsonSerializerMacros.h"
+
+const FName AssetActionUtilityTags::BlutilityTagVersion(TEXT("BlutilityTagVersion"));
+const FName AssetActionUtilityTags::SupportedClasses(TEXT("SupportedClasses"));
+const FName AssetActionUtilityTags::IsActionForBlueprint(TEXT("IsActionForBlueprint"));
+const FName AssetActionUtilityTags::CallableFunctions(TEXT("CallableFunctions"));
+
+const int32 AssetActionUtilityTags::TagVersion = 1;
+
+UObject* FAssetActionUtilityPrototype::LoadUtilityAsset() const
+{
+	if (const UEditorUtilityBlueprint* Blueprint = Cast<UEditorUtilityBlueprint>(UtilityBlueprintAsset.GetAsset()))
+	{
+		if (const UClass* BPClass = Blueprint->GeneratedClass.Get())
+		{
+			return BPClass->GetDefaultObject();
+		}
+	}
+
+	return nullptr;
+};
+
+bool FAssetActionUtilityPrototype::IsLatestVersion() const
+{
+	int32 Version;
+    if (UtilityBlueprintAsset.GetTagValue(AssetActionUtilityTags::BlutilityTagVersion, Version))
+    {
+    	return Version == AssetActionUtilityTags::TagVersion;
+    }
+    
+    return false;
+}
+
+bool FAssetActionUtilityPrototype::AreSupportedClassesForBlueprints() const
+{
+	FString IsActionForBlueprintString;
+	if (UtilityBlueprintAsset.GetTagValue(AssetActionUtilityTags::IsActionForBlueprint, IsActionForBlueprintString))
+	{
+		bool IsActionForBlueprint = false;
+		LexFromString(IsActionForBlueprint, *IsActionForBlueprintString);
+		return IsActionForBlueprint;
+	}
+
+	return false;
+}
+
+TArray<TSoftClassPtr<UObject>> FAssetActionUtilityPrototype::GetSupportedClasses() const
+{
+	FString SupportedClassPaths;
+	if (UtilityBlueprintAsset.GetTagValue(AssetActionUtilityTags::SupportedClasses, SupportedClassPaths))
+	{
+		TArray<FString> ClassPaths;
+		SupportedClassPaths.ParseIntoArray(ClassPaths, TEXT(","));
+		TArray<TSoftClassPtr<UObject>> ClassPtrs;
+		Algo::Transform(ClassPaths, ClassPtrs, [](const FString& ClassPath) { return TSoftClassPtr<UObject>(FSoftClassPath(ClassPath)); });
+		return ClassPtrs;
+	}
+
+	return TArray<TSoftClassPtr<UObject>>();
+}
+
+TArray<FBlutilityFunctionData> FAssetActionUtilityPrototype::GetCallableFunctions() const
+{
+	TArray<FBlutilityFunctionData> FunctionDatas;
+	
+	FString FunctionDataJson;
+	if (UtilityBlueprintAsset.GetTagValue(AssetActionUtilityTags::CallableFunctions, FunctionDataJson))
+	{
+		FJsonObjectConverter::JsonArrayStringToUStruct(FunctionDataJson, &FunctionDatas, 0, 0);
+	}
+
+	return FunctionDatas;
+}
+
+void FAssetActionUtilityPrototype::AddTagsFor_Version(TArray<UObject::FAssetRegistryTag>& OutTags)
+{
+	// Adding a version to the tags just in case we need to know 'this blutility is out of date and we can't go based off the cached data because the format isn't followed, or this is a pre-tagged version'
+	OutTags.Add(UObject::FAssetRegistryTag(AssetActionUtilityTags::BlutilityTagVersion, LexToString(AssetActionUtilityTags::TagVersion), UObject::FAssetRegistryTag::TT_Hidden));	
+}
+
+void FAssetActionUtilityPrototype::AddTagsFor_SupportedClasses(const TArray<TSoftClassPtr<UObject>>& SupportedClasses, TArray<UObject::FAssetRegistryTag>& OutTags)
+{
+	const FString SupportedClassesString = FString::JoinBy(SupportedClasses, TEXT(","), [](const TSoftClassPtr<UObject>& SupportedClass) { return SupportedClass.ToString(); });
+	OutTags.Add(UObject::FAssetRegistryTag(AssetActionUtilityTags::SupportedClasses, SupportedClassesString, UObject::FAssetRegistryTag::TT_Hidden));
+}
+
+void FAssetActionUtilityPrototype::AddTagsFor_IsActionForBlueprints(bool IsActionForBlueprints, TArray<UObject::FAssetRegistryTag>& OutTags)
+{
+	const bool ActionForBlueprint = IsActionForBlueprints;
+	OutTags.Add(UObject::FAssetRegistryTag(AssetActionUtilityTags::IsActionForBlueprint, ActionForBlueprint ? TEXT("True") : TEXT("False"), UObject::FAssetRegistryTag::TT_Hidden));
+}
+
+void FAssetActionUtilityPrototype::AddTagsFor_CallableFunctions(const UObject* FunctionsSource, TArray<UObject::FAssetRegistryTag>& OutTags)
+{
+	check(FunctionsSource);
+	
+	const static FName NAME_CallInEditor(TEXT("CallInEditor"));
+	const static FName NAME_Category(TEXT("Category"));
+
+	TArray<FBlutilityFunctionData> FunctionDatas;
+	for (TFieldIterator<UFunction> FunctionIt(FunctionsSource->GetClass(), EFieldIterationFlags::IncludeSuper); FunctionIt; ++FunctionIt)
+	{
+		if (UFunction* Func = *FunctionIt)
+		{
+			if (Func->HasMetaData(NAME_CallInEditor) && Func->GetReturnProperty() == nullptr)
+			{
+				FBlutilityFunctionData Function;
+				Function.Class = Func->GetOuterUClass();
+				Function.Name = Func->GetFName();
+				Function.NameText = Func->GetDisplayNameText();
+				Function.Category = Func->GetMetaData(NAME_Category);
+				Function.TooltipText = Func->GetToolTipText();
+				
+				FunctionDatas.Add(MoveTemp(Function));
+			}
+		}
+	}
+
+	TArray<TSharedPtr<FJsonValue>> FunctionDataJsonValus;
+	for (const FBlutilityFunctionData& Highlight : FunctionDatas)
+	{
+		TSharedPtr<FJsonObject> FunctionDataJO = FJsonObjectConverter::UStructToJsonObject(Highlight);
+		if (FunctionDataJO.IsValid())
+		{
+			FunctionDataJsonValus.Add(MakeShared<FJsonValueObject>(FunctionDataJO));
+		}
+	}
+
+	FString FunctionDataJson;
+	{
+		TSharedRef<FJsonValueArray> HighlightValuesArrayValue = MakeShared<FJsonValueArray>(FunctionDataJsonValus);
+		TSharedRef<TJsonWriter<TCHAR, TCondensedJsonPrintPolicy<TCHAR>>> JsonWriter = TJsonWriterFactory<TCHAR, TCondensedJsonPrintPolicy<TCHAR>>::Create(&FunctionDataJson);
+		FJsonSerializer::Serialize(HighlightValuesArrayValue, FString(), JsonWriter);
+	}
+
+	OutTags.Add(UObject::FAssetRegistryTag(AssetActionUtilityTags::CallableFunctions, FunctionDataJson, UObject::FAssetRegistryTag::TT_Hidden));
+}
\ No newline at end of file
diff --git a/Engine/Source/Editor/Blutility/Private/EditorUtilityAssetPrototype.h b/Engine/Source/Editor/Blutility/Private/EditorUtilityAssetPrototype.h
new file mode 100644
--- /dev/null
+++ b/Engine/Source/Editor/Blutility/Private/EditorUtilityAssetPrototype.h
@@ -0,0 +1,93 @@
+// Copyright Epic Games, Inc. All Rights Reserved.
+
+#pragma once
+
+#include "UObject/Object.h"
+#include "AssetRegistry/AssetData.h"
+#include "IEditorUtilityExtension.h"
+
+#include "EditorUtilityAssetPrototype.generated.h"
+
+USTRUCT()
+struct FBlutilityFunctionData
+{
+	GENERATED_BODY()
+
+public:
+	UPROPERTY()
+	TSoftClassPtr<UClass> Class;
+	
+	UPROPERTY()
+	FName Name;
+	
+	UPROPERTY()
+	FText NameText;
+
+	UPROPERTY()
+	FString Category;
+
+	UPROPERTY()
+	FText TooltipText;
+
+	bool operator !=(const FBlutilityFunctionData& Other) const
+	{
+		return !(*this == Other);
+	}
+
+	bool operator == (const FBlutilityFunctionData& Other) const
+	{
+		return Class == Other.Class && Name == Other.Name;
+	}
+};
+
+namespace AssetActionUtilityTags
+{
+	extern const FName BlutilityTagVersion;
+	extern const FName SupportedClasses;
+	extern const FName IsActionForBlueprint;
+	extern const FName CallableFunctions;
+	
+	extern const int32 TagVersion;
+}
+
+class FAssetActionUtilityPrototype
+{
+public:
+	FAssetActionUtilityPrototype(const FAssetData& SourceAsset)
+		: UtilityBlueprintAsset(SourceAsset)
+	{
+	}
+
+	UObject* LoadUtilityAsset() const;
+
+	bool IsLatestVersion() const;
+	bool AreSupportedClassesForBlueprints() const;
+	TArray<TSoftClassPtr<UObject>> GetSupportedClasses() const;
+	TArray<FBlutilityFunctionData> GetCallableFunctions() const;
+
+	const FAssetData& GetUtilityBlueprintAsset() const { return UtilityBlueprintAsset; }
+
+	bool operator !=(const FAssetActionUtilityPrototype& Other) const
+	{
+		return !(*this == Other);
+	}
+
+	bool operator == (const FAssetActionUtilityPrototype& Other) const
+	{
+		return UtilityBlueprintAsset.PackageName == Other.UtilityBlueprintAsset.PackageName;
+	}
+
+public:
+	static void AddTagsFor_Version(TArray<UObject::FAssetRegistryTag>& OutTags);
+	static void AddTagsFor_SupportedClasses(const TArray<TSoftClassPtr<UObject>>& SupportedClasses, TArray<UObject::FAssetRegistryTag>& OutTags);
+	static void AddTagsFor_IsActionForBlueprints(bool IsActionForBlueprints, TArray<UObject::FAssetRegistryTag>& OutTags);
+	static void AddTagsFor_CallableFunctions(const UObject* FunctionsSource, TArray<UObject::FAssetRegistryTag>& OutTags);
+
+private:
+	FAssetData UtilityBlueprintAsset;
+};
+
+FORCEINLINE uint32 GetTypeHash(const FAssetActionUtilityPrototype& ExtensionAsset)
+{
+	return HashCombine(GetTypeHash(ExtensionAsset.GetUtilityBlueprintAsset().PackageName), GetTypeHash(ExtensionAsset.GetUtilityBlueprintAsset().AssetName));
+}
\ No newline at end of file
diff --git a/Engine/Source/Editor/Blutility/Private/EditorUtilityLibrary.cpp b/Engine/Source/Editor/Blutility/Private/EditorUtilityLibrary.cpp
--- a/Engine/Source/Editor/Blutility/Private/EditorUtilityLibrary.cpp
+++ b/Engine/Source/Editor/Blutility/Private/EditorUtilityLibrary.cpp
@@ -238,6 +238,11 @@ void UEditorUtilityLibrary::GetSelectionBounds(FVector& Origin, FVector& BoxExte
 }
 
 TArray<UObject*> UEditorUtilityLibrary::GetSelectedAssets()
+{
+	return GetSelectedAssetsOfClass(UObject::StaticClass());
+}
+
+TArray<UObject*> UEditorUtilityLibrary::GetSelectedAssetsOfClass(UClass* AssetClass)
 {
 	//@TODO: Blocking load, no slow dialog
 	FContentBrowserModule& ContentBrowserModule = FModuleManager::LoadModuleChecked<FContentBrowserModule>("ContentBrowser");
@@ -247,7 +252,10 @@ TArray<UObject*> UEditorUtilityLibrary::GetSelectedAssets()
 	TArray<UObject*> Result;
 	for (FAssetData& AssetData : SelectedAssets)
 	{
-		Result.Add(AssetData.GetAsset());
+		if (AssetData.IsInstanceOf(AssetClass))
+		{
+			Result.Add(AssetData.GetAsset());
+		}
 	}
 
 	return Result;
diff --git a/Engine/Source/Editor/Blutility/Private/IEditorUtilityExtension.h b/Engine/Source/Editor/Blutility/Private/IEditorUtilityExtension.h
new file mode 100644
--- /dev/null
+++ b/Engine/Source/Editor/Blutility/Private/IEditorUtilityExtension.h
@@ -0,0 +1,20 @@
+// Copyright Epic Games, Inc. All Rights Reserved.
+
+#pragma once
+
+#include "CoreMinimal.h"
+#include "UObject/ObjectMacros.h"
+#include "UObject/Interface.h"
+
+#include "IEditorUtilityExtension.generated.h"
+
+UINTERFACE(BlueprintType)
+class UEditorUtilityExtension : public UInterface
+{
+	GENERATED_UINTERFACE_BODY()
+};
+
+class IEditorUtilityExtension
+{
+	GENERATED_IINTERFACE_BODY()
+};
diff --git a/Engine/Source/Editor/Blutility/Public/EditorUtilityLibrary.h b/Engine/Source/Editor/Blutility/Public/EditorUtilityLibrary.h
--- a/Engine/Source/Editor/Blutility/Public/EditorUtilityLibrary.h
+++ b/Engine/Source/Editor/Blutility/Public/EditorUtilityLibrary.h
@@ -148,6 +148,9 @@ public:
 	// Gets the set of currently selected assets
 	UFUNCTION(BlueprintCallable, Category = "Development|Editor")
 	static TArray<UObject*> GetSelectedAssets();
+	
+	UFUNCTION(BlueprintCallable, Category = "Development|Editor")
+	static TArray<UObject*> GetSelectedAssetsOfClass(UClass* AssetClass);
 
 	// Gets the set of currently selected classes
 	UFUNCTION(BlueprintCallable, Category = "Development|Editor")
diff --git a/Engine/Source/Editor/ContentBrowser/Private/AssetContextMenu.cpp b/Engine/Source/Editor/ContentBrowser/Private/AssetContextMenu.cpp
--- a/Engine/Source/Editor/ContentBrowser/Private/AssetContextMenu.cpp
+++ b/Engine/Source/Editor/ContentBrowser/Private/AssetContextMenu.cpp
@@ -92,61 +92,8 @@ void FAssetContextMenu::BindCommands(TSharedPtr< FUICommandList >& Commands)
 		));
 }
 
-class FDontLoadAssetsDuringThisScope
-{
-public:
-	FDontLoadAssetsDuringThisScope()
-	{
-		FCoreUObjectDelegates::OnAssetLoaded.AddRaw(this, &FDontLoadAssetsDuringThisScope::OnAssetLoaded);
-	}
-
-	~FDontLoadAssetsDuringThisScope()
-	{
-		FCoreUObjectDelegates::OnAssetLoaded.RemoveAll(this);
-
-		FMessageLog EditorErrors("EditorErrors");
-		
-		for (const FLoadEvent& LoadEvent : ObjectsLoaded)
-		{
-			TSharedRef<FTokenizedMessage> ErrorMessage = EditorErrors.Error();
-			ErrorMessage->AddToken(FAssetNameToken::Create(LoadEvent.ObjectPath.ToString(),FText::FromString(LoadEvent.ObjectPath.ToString())));
-			ErrorMessage->AddToken(FTextToken::Create(LOCTEXT("DontLoadAssetsNow", "was loaded while creating the asset context menu.  Don't do this, it causes the operation to hitch.")));
-			
-			for (const FProgramCounterSymbolInfo& CallInfo : LoadEvent.Callstack)
-			{
-				ErrorMessage->AddToken(FTextToken::Create(FText::FromString(FString(CallInfo.FunctionName) + TEXT("\n"))));
-			}
-		}
-
-		EditorErrors.Notify(LOCTEXT("AssetsOpenedDuringContextMenuNotify", "Don't Load Assets Now!"));
-	}
-	
-private:
-	void OnAssetLoaded(UObject* InObject)
-	{
-		FLoadEvent Event;
-		Event.ObjectPath = FSoftObjectPath(InObject);
-		Event.Callstack = FPlatformStackWalk::GetStack(4, 10);
-		ObjectsLoaded.Add(MoveTemp(Event));
-	}
-
-private:
-	struct FLoadEvent
-	{
-		FSoftObjectPath ObjectPath;
-		TArray<FProgramCounterSymbolInfo> Callstack;
-	};
-	
-	TArray<FLoadEvent> ObjectsLoaded;
-};
-
 TSharedRef<SWidget> FAssetContextMenu::MakeContextMenu(TArrayView<const FContentBrowserItem> InSelectedItems, const FSourcesData& InSourcesData, TSharedPtr< FUICommandList > InCommandList)
 {
-#if DEPRECATE_ASSET_TYPE_ACTIONS_CALLING_GETACTIONS
-	// If we're not calling GetActions any more, then fire off errors when people open the context menu and it loads assets
-	FDontLoadAssetsDuringThisScope DontLoadScope;
-#endif
-	
 	SetSelectedItems(InSelectedItems);
 	SourcesData = InSourcesData;
 
@@ -377,32 +324,31 @@ void FAssetContextMenu::RegisterContextMenu(const FName MenuName)
 		UToolMenu* Menu = ToolMenus->RegisterMenu(MenuName);
 		FToolMenuSection& Section = Menu->FindOrAddSection("GetAssetActions");
 
-#if  !DEPRECATE_ASSET_TYPE_ACTIONS_CALLING_GETACTIONS
-
-		// Note: Do  not use "GetActions" again when copying this code, otherwise "GetActions" menu entry will be overwritten
-		Section.AddDynamicEntry("GetActions", FNewToolMenuSectionDelegate::CreateLambda([](FToolMenuSection& InSection)
+		// TODO Remove when IAssetTypeActions is dead or fully deprecated.
 		{
-			UContentBrowserAssetContextMenuContext* Context = InSection.FindContext<UContentBrowserAssetContextMenuContext>();
-			if (Context && Context->CommonAssetTypeActions.IsValid())
+			// Note: Do  not use "GetActions" again when copying this code, otherwise "GetActions" menu entry will be overwritten
+			Section.AddDynamicEntry("GetActions", FNewToolMenuSectionDelegate::CreateLambda([](FToolMenuSection& InSection)
 			{
-				PRAGMA_DISABLE_DEPRECATION_WARNINGS
-				Context->CommonAssetTypeActions.Pin()->GetActions(Context->LoadSelectedObjectsIfNeeded(), InSection);
-				PRAGMA_ENABLE_DEPRECATION_WARNINGS
-			}
-		}));
+				UContentBrowserAssetContextMenuContext* Context = InSection.FindContext<UContentBrowserAssetContextMenuContext>();
+				if (Context && Context->CommonAssetTypeActions.IsValid())
+				{
+					PRAGMA_DISABLE_DEPRECATION_WARNINGS
+					Context->CommonAssetTypeActions.Pin()->GetActions(Context->LoadSelectedObjectsIfNeeded(), InSection);
+					PRAGMA_ENABLE_DEPRECATION_WARNINGS
+				}
+			}));
 
-		Section.AddDynamicEntry("GetActionsLegacy", FNewToolMenuDelegateLegacy::CreateLambda([](FMenuBuilder& MenuBuilder, UToolMenu* InMenu)
-		{
-			UContentBrowserAssetContextMenuContext* Context = InMenu->FindContext<UContentBrowserAssetContextMenuContext>();
-			if (Context && Context->CommonAssetTypeActions.IsValid())
+			Section.AddDynamicEntry("GetActionsLegacy", FNewToolMenuDelegateLegacy::CreateLambda([](FMenuBuilder& MenuBuilder, UToolMenu* InMenu)
 			{
-				PRAGMA_DISABLE_DEPRECATION_WARNINGS
-				Context->CommonAssetTypeActions.Pin()->GetActions(Context->LoadSelectedObjectsIfNeeded(), MenuBuilder);
-				PRAGMA_ENABLE_DEPRECATION_WARNINGS
-			}
-		}));
-		
-#endif
+				UContentBrowserAssetContextMenuContext* Context = InMenu->FindContext<UContentBrowserAssetContextMenuContext>();
+				if (Context && Context->CommonAssetTypeActions.IsValid())
+				{
+					PRAGMA_DISABLE_DEPRECATION_WARNINGS
+					Context->CommonAssetTypeActions.Pin()->GetActions(Context->LoadSelectedObjectsIfNeeded(), MenuBuilder);
+					PRAGMA_ENABLE_DEPRECATION_WARNINGS
+				}
+			}));
+		}		
 
 		Menu->AddDynamicSection("AddMenuOptions", FNewToolMenuDelegate::CreateLambda([](UToolMenu* InMenu)
 		{
diff --git a/Engine/Source/Runtime/Engine/Classes/Engine/Blueprint.h b/Engine/Source/Runtime/Engine/Classes/Engine/Blueprint.h
--- a/Engine/Source/Runtime/Engine/Classes/Engine/Blueprint.h
+++ b/Engine/Source/Runtime/Engine/Classes/Engine/Blueprint.h
@@ -1145,7 +1145,7 @@ public:
 
 		return false;
 	}
-
+	
 #if WITH_EDITOR
 	static FName GetFunctionNameFromClassByGuid(const UClass* InClass, const FGuid FunctionGuid);
 	static bool GetFunctionGuidFromClassByFieldName(const UClass* InClass, const FName FunctionName, FGuid& FunctionGuid);
@@ -1160,6 +1160,10 @@ public:
 	virtual void NotifyGraphRenamed(class UEdGraph* Graph, FName OldName, FName NewName) { }
 #endif
 
+#if WITH_EDITOR
+	static UClass* GetBlueprintParentClassFromAssetTags(const FAssetData& BlueprintAsset);
+#endif
+
 	/** Find a function given its name and optionally an object property name within this Blueprint */
 	ETimelineSigType GetTimelineSignatureForFunctionByName(const FName& FunctionName, const FName& ObjectPropertyName);
 
diff --git a/Engine/Source/Runtime/Engine/Private/Blueprint.cpp b/Engine/Source/Runtime/Engine/Private/Blueprint.cpp
--- a/Engine/Source/Runtime/Engine/Private/Blueprint.cpp
+++ b/Engine/Source/Runtime/Engine/Private/Blueprint.cpp
@@ -2094,6 +2094,29 @@ UEdGraph* UBlueprint::GetLastEditedUberGraph() const
 	return nullptr;
 }
 
+#if WITH_EDITOR
+
+UClass* UBlueprint::GetBlueprintParentClassFromAssetTags(const FAssetData& BlueprintAsset)
+{
+	UClass* ParentClass = nullptr;
+	FString ParentClassName;
+	if(!BlueprintAsset.GetTagValue(FBlueprintTags::NativeParentClassPath, ParentClassName))
+	{
+		BlueprintAsset.GetTagValue(FBlueprintTags::ParentClassPath, ParentClassName);
+	}
+	
+	if(!ParentClassName.IsEmpty())
+	{
+		UObject* Outer = nullptr;
+		ResolveName(Outer, ParentClassName, false, false);
+		ParentClass = FindObject<UClass>(Outer, *ParentClassName);
+	}
+	
+	return ParentClass;
+}
+
+#endif
+
 PRAGMA_DISABLE_DEPRECATION_WARNINGS
 
 TArrayView<const TObjectPtr<UBlueprintExtension>> UBlueprint::GetExtensions() const
