Editor - Fixing an SA warning about comparing a value to itself.

diff --git a/Engine/Source/Editor/Blutility/Private/BlutilityMenuExtensions.h b/Engine/Source/Editor/Blutility/Private/BlutilityMenuExtensions.h
--- a/Engine/Source/Editor/Blutility/Private/BlutilityMenuExtensions.h
+++ b/Engine/Source/Editor/Blutility/Private/BlutilityMenuExtensions.h
@@ -48,7 +48,7 @@ protected:
 
 		bool operator==(const FFunctionAndUtil& InFunction) const
 		{
-			return InFunction.FunctionData == InFunction.FunctionData && InFunction.Util == Util;
+			return FunctionData == InFunction.FunctionData && Util == InFunction.Util;
 		}
 
 		UFunction* GetFunction() const
