Editor - Updating the deprecation messages on AssetActionUtility and ActorActionUtility for the SupportedClasses array to help people find it easier.

diff --git a/Engine/Source/Editor/Blutility/Classes/ActorActionUtility.h b/Engine/Source/Editor/Blutility/Classes/ActorActionUtility.h
--- a/Engine/Source/Editor/Blutility/Classes/ActorActionUtility.h
+++ b/Engine/Source/Editor/Blutility/Classes/ActorActionUtility.h
@@ -19,8 +19,8 @@ class BLUTILITY_API UActorActionUtility : public UEditorUtilityObject, public IE
 
 public:
 	/** Return the class that this actor action supports. Leave this blank to support all actor classes. */
-	UE_DEPRECATED(5.2, "GetSupportedClasses() instead, but ideally you're not requesting this directly and are instead using GetSupportedClasses_FromAssetTags.")
-	UFUNCTION(BlueprintPure, BlueprintImplementableEvent, Category="Assets", meta=(DeprecatedFunction, DeprecationMessage="If you were just returning a single class.  Add it to the SupportedClasses array.  If you were doing complex logic to simulate having multiple classes act as filters, add them to the SupportedClasses array.  If you were doing 'other' logic, you'll need to do that upon action execution."))
+	UE_DEPRECATED(5.2, "GetSupportedClasses() instead, but ideally you're not requesting this directly and are instead using the FAssetActionUtilityPrototype to wrap access to an unload utility asset.")
+	UFUNCTION(BlueprintPure, BlueprintImplementableEvent, Category="Assets", meta=(DeprecatedFunction, DeprecationMessage="If you were just returning a single class add it to the SupportedClasses array (you can find it listed in the Class Defaults).  If you were doing complex logic to simulate having multiple classes act as filters, add them to the SupportedClasses array.  If you were doing 'other' logic, you'll need to do that upon action execution."))
 	UClass* GetSupportedClass() const;
 	
 	/**
diff --git a/Engine/Source/Editor/Blutility/Classes/AssetActionUtility.h b/Engine/Source/Editor/Blutility/Classes/AssetActionUtility.h
--- a/Engine/Source/Editor/Blutility/Classes/AssetActionUtility.h
+++ b/Engine/Source/Editor/Blutility/Classes/AssetActionUtility.h
@@ -33,8 +33,8 @@ public:
 	 * Return the class that this asset action supports (if not implemented, it will show up for all asset types)
 	 * Do not do custom logic here based on the currently selected assets.
 	 */
-	UE_DEPRECATED(5.2, "GetSupportedClasses() instead, but ideally you're not requesting this directly and are instead using GetSupportedClasses_FromAssetTags.")
-	UFUNCTION(BlueprintPure, BlueprintImplementableEvent, Category="Assets", meta=(DeprecatedFunction, DeprecationMessage="If you were just returning a single class.  Add it to the SupportedClasses array.  If you were doing complex logic to simulate having multiple classes act as filters, add them to the SupportedClasses array.  If you were doing 'other' logic, you'll need to do that upon action execution."))
+	UE_DEPRECATED(5.2, "GetSupportedClasses() instead, but ideally you're not requesting this directly and are instead using the FAssetActionUtilityPrototype to wrap access to an unload utility asset.")
+	UFUNCTION(BlueprintPure, BlueprintImplementableEvent, Category="Assets", meta=(DeprecatedFunction, DeprecationMessage="If you were just returning a single class add it to the SupportedClasses array (you can find it listed in the Class Defaults).  If you were doing complex logic to simulate having multiple classes act as filters, add them to the SupportedClasses array.  If you were doing 'other' logic, you'll need to do that upon action execution."))
 	UClass* GetSupportedClass() const;
 
 	/**
