Editor - Protecting from a crash by checking if we actually have any function data before serializing it to json for the utility script function metadata.

diff --git a/Engine/Source/Editor/Blutility/Private/EditorUtilityAssetPrototype.cpp b/Engine/Source/Editor/Blutility/Private/EditorUtilityAssetPrototype.cpp
--- a/Engine/Source/Editor/Blutility/Private/EditorUtilityAssetPrototype.cpp
+++ b/Engine/Source/Editor/Blutility/Private/EditorUtilityAssetPrototype.cpp
@@ -122,22 +122,25 @@ void FAssetActionUtilityPrototype::AddTagsFor_CallableFunctions(const UObject* F
 		}
 	}
 
-	TArray<TSharedPtr<FJsonValue>> FunctionDataJsonValus;
+	TArray<TSharedPtr<FJsonValue>> FunctionDataJsonValues;
 	for (const FBlutilityFunctionData& Highlight : FunctionDatas)
 	{
 		TSharedPtr<FJsonObject> FunctionDataJO = FJsonObjectConverter::UStructToJsonObject(Highlight);
 		if (FunctionDataJO.IsValid())
 		{
-			FunctionDataJsonValus.Add(MakeShared<FJsonValueObject>(FunctionDataJO));
+			FunctionDataJsonValues.Add(MakeShared<FJsonValueObject>(FunctionDataJO));
 		}
 	}
 
-	FString FunctionDataJson;
+	if (FunctionDataJsonValues.Num() > 0)
 	{
-		TSharedRef<FJsonValueArray> HighlightValuesArrayValue = MakeShared<FJsonValueArray>(FunctionDataJsonValus);
-		TSharedRef<TJsonWriter<TCHAR, TCondensedJsonPrintPolicy<TCHAR>>> JsonWriter = TJsonWriterFactory<TCHAR, TCondensedJsonPrintPolicy<TCHAR>>::Create(&FunctionDataJson);
-		FJsonSerializer::Serialize(HighlightValuesArrayValue, FString(), JsonWriter);
-	}
+		FString FunctionDataJson;
+		{
+			TSharedRef<FJsonValueArray> HighlightValuesArrayValue = MakeShared<FJsonValueArray>(FunctionDataJsonValues);
+			TSharedRef<TJsonWriter<TCHAR, TCondensedJsonPrintPolicy<TCHAR>>> JsonWriter = TJsonWriterFactory<TCHAR, TCondensedJsonPrintPolicy<TCHAR>>::Create(&FunctionDataJson);
+			FJsonSerializer::Serialize(HighlightValuesArrayValue, FString(), JsonWriter);
+		}
 
-	OutTags.Add(UObject::FAssetRegistryTag(AssetActionUtilityTags::CallableFunctions, FunctionDataJson, UObject::FAssetRegistryTag::TT_Hidden));
+		OutTags.Add(UObject::FAssetRegistryTag(AssetActionUtilityTags::CallableFunctions, FunctionDataJson, UObject::FAssetRegistryTag::TT_Hidden));
+	}
 }
\ No newline at end of file
