Blutility - Adding a fallback for cases where no supported classes are serialized to tag data.  When that happens the assumption has always been that you support everything, and so we just add UObject to the array when it's empty.

diff --git a/Engine/Source/Editor/Blutility/Private/EditorUtilityAssetPrototype.cpp b/Engine/Source/Editor/Blutility/Private/EditorUtilityAssetPrototype.cpp
--- a/Engine/Source/Editor/Blutility/Private/EditorUtilityAssetPrototype.cpp
+++ b/Engine/Source/Editor/Blutility/Private/EditorUtilityAssetPrototype.cpp
@@ -52,17 +52,22 @@ bool FAssetActionUtilityPrototype::AreSupportedClassesForBlueprints() const
 
 TArray<TSoftClassPtr<UObject>> FAssetActionUtilityPrototype::GetSupportedClasses() const
 {
+    TArray<TSoftClassPtr<UObject>> SupportedClasses;
+
 	FString SupportedClassPaths;
 	if (UtilityBlueprintAsset.GetTagValue(AssetActionUtilityTags::SupportedClasses, SupportedClassPaths))
 	{
 		TArray<FString> ClassPaths;
 		SupportedClassPaths.ParseIntoArray(ClassPaths, TEXT(","));
-		TArray<TSoftClassPtr<UObject>> ClassPtrs;
-		Algo::Transform(ClassPaths, ClassPtrs, [](const FString& ClassPath) { return TSoftClassPtr<UObject>(FSoftClassPath(ClassPath)); });
-		return ClassPtrs;
+		Algo::Transform(ClassPaths, SupportedClasses, [](const FString& ClassPath) { return TSoftClassPtr<UObject>(FSoftClassPath(ClassPath)); });
 	}
 
-	return TArray<TSoftClassPtr<UObject>>();
+	if (SupportedClasses.Num() == 0)
+	{
+		SupportedClasses.Add(TSoftClassPtr<UObject>(UObject::StaticClass()));
+	}
+	
+	return SupportedClasses;
 }
 
 TArray<FBlutilityFunctionData> FAssetActionUtilityPrototype::GetCallableFunctions() const
