Blutility - we now pass in EResolveClass::Yes when checking if the class instance matches for blutility scripts on assets.  Otherwise your script commands may not show up as available on unload script class based assets.

diff --git a/Engine/Source/Editor/Blutility/Private/BlutilityContentBrowserExtensions.cpp b/Engine/Source/Editor/Blutility/Private/BlutilityContentBrowserExtensions.cpp
--- a/Engine/Source/Editor/Blutility/Private/BlutilityContentBrowserExtensions.cpp
+++ b/Engine/Source/Editor/Blutility/Private/BlutilityContentBrowserExtensions.cpp
@@ -78,7 +78,7 @@ public:
 							{
 								// Is the asset the right kind?
 								bPassesClassFilter = 
-									Algo::AnyOf(SupportedClassPtrs, [&Asset](TSoftClassPtr<UObject> ClassPtr){ return Asset.IsInstanceOf(ClassPtr.Get()); });
+									Algo::AnyOf(SupportedClassPtrs, [&Asset](TSoftClassPtr<UObject> ClassPtr){ return Asset.IsInstanceOf(ClassPtr.Get(), EResolveClass::Yes); });
 							}
 
 							if (bPassesClassFilter)
