Fix to Editor Utility to check for null ptr

diff --git a/Engine/Source/Editor/Blutility/Private/EditorUtilitySubsystem.cpp b/Engine/Source/Editor/Blutility/Private/EditorUtilitySubsystem.cpp
--- a/Engine/Source/Editor/Blutility/Private/EditorUtilitySubsystem.cpp
+++ b/Engine/Source/Editor/Blutility/Private/EditorUtilitySubsystem.cpp
@@ -226,17 +226,19 @@ void UEditorUtilitySubsystem::RegisterTabAndGetID(class UEditorUtilityWidgetBlue
 		FName RegistrationName = NewTabID.IsNone() ? FName(*(InBlueprint->GetPathName() + LOCTEXT("ActiveTabSuffix", "_ActiveTab").ToString())) : FName(*(InBlueprint->GetPathName() + NewTabID.ToString()));
 		FText DisplayName = FText::FromString(FName::NameToDisplayString(InBlueprint->GetName(), false));
 		FLevelEditorModule& LevelEditorModule = FModuleManager::GetModuleChecked<FLevelEditorModule>(TEXT("LevelEditor"));
-		TSharedPtr<FTabManager> LevelEditorTabManager = LevelEditorModule.GetLevelEditorTabManager();
-		if (!LevelEditorTabManager->HasTabSpawner(RegistrationName))
+		if (TSharedPtr<FTabManager> LevelEditorTabManager = LevelEditorModule.GetLevelEditorTabManager())
 		{
-			IBlutilityModule* BlutilityModule = FModuleManager::GetModulePtr<IBlutilityModule>("Blutility");
-			LevelEditorTabManager->RegisterTabSpawner(RegistrationName, FOnSpawnTab::CreateUObject(InBlueprint, &UEditorUtilityWidgetBlueprint::SpawnEditorUITab))
-				.SetDisplayName(DisplayName)
-				.SetGroup(BlutilityModule->GetMenuGroup().ToSharedRef());
-			InBlueprint->SetRegistrationName(RegistrationName);
+			if (!LevelEditorTabManager->HasTabSpawner(RegistrationName))
+			{
+				IBlutilityModule* BlutilityModule = FModuleManager::GetModulePtr<IBlutilityModule>("Blutility");
+				LevelEditorTabManager->RegisterTabSpawner(RegistrationName, FOnSpawnTab::CreateUObject(InBlueprint, &UEditorUtilityWidgetBlueprint::SpawnEditorUITab))
+					.SetDisplayName(DisplayName)
+					.SetGroup(BlutilityModule->GetMenuGroup().ToSharedRef());
+				InBlueprint->SetRegistrationName(RegistrationName);
+			}
+			RegisteredTabs.Add(RegistrationName, InBlueprint);
+			NewTabID = RegistrationName;
 		}
-		RegisteredTabs.Add(RegistrationName, InBlueprint);
-		NewTabID = RegistrationName;
 	}
 }
 
@@ -245,15 +247,17 @@ bool UEditorUtilitySubsystem::SpawnRegisteredTabByID(FName NewTabID)
 	if (!IsRunningCommandlet())
 	{
 		FLevelEditorModule& LevelEditorModule = FModuleManager::GetModuleChecked<FLevelEditorModule>(TEXT("LevelEditor"));
-		TSharedPtr<FTabManager> LevelEditorTabManager = LevelEditorModule.GetLevelEditorTabManager();
-		if (LevelEditorTabManager->HasTabSpawner(NewTabID))
+		if (TSharedPtr<FTabManager> LevelEditorTabManager = LevelEditorModule.GetLevelEditorTabManager())
 		{
-			TSharedPtr<SDockTab> NewDockTab = LevelEditorTabManager->TryInvokeTab(NewTabID);
-			IBlutilityModule* BlutilityModule = FModuleManager::GetModulePtr<IBlutilityModule>("Blutility");
-			UEditorUtilityWidgetBlueprint* WidgetToSpawn = *RegisteredTabs.Find(NewTabID);
-			check(WidgetToSpawn);
-			BlutilityModule->AddLoadedScriptUI(WidgetToSpawn);
-			return true;
+			if (LevelEditorTabManager->HasTabSpawner(NewTabID))
+			{
+				TSharedPtr<SDockTab> NewDockTab = LevelEditorTabManager->TryInvokeTab(NewTabID);
+				IBlutilityModule* BlutilityModule = FModuleManager::GetModulePtr<IBlutilityModule>("Blutility");
+				UEditorUtilityWidgetBlueprint* WidgetToSpawn = *RegisteredTabs.Find(NewTabID);
+				check(WidgetToSpawn);
+				BlutilityModule->AddLoadedScriptUI(WidgetToSpawn);
+				return true;
+			}
 		}
 	}
 	return false;
@@ -264,11 +268,13 @@ bool UEditorUtilitySubsystem::DoesTabExist(FName NewTabID)
 	if (!IsRunningCommandlet())
 	{
 		FLevelEditorModule& LevelEditorModule = FModuleManager::GetModuleChecked<FLevelEditorModule>(TEXT("LevelEditor"));
-		TSharedPtr<FTabManager> LevelEditorTabManager = LevelEditorModule.GetLevelEditorTabManager();
-		TSharedPtr<SDockTab> FoundTab = LevelEditorTabManager->FindExistingLiveTab(NewTabID);
-		if (FoundTab)
+		if (TSharedPtr<FTabManager> LevelEditorTabManager = LevelEditorModule.GetLevelEditorTabManager())
 		{
-			return true;
+			TSharedPtr<SDockTab> FoundTab = LevelEditorTabManager->FindExistingLiveTab(NewTabID);
+			if (FoundTab)
+			{
+				return true;
+			}
 		}
 	}
 
@@ -280,12 +286,14 @@ bool UEditorUtilitySubsystem::CloseTabByID(FName NewTabID)
 	if (!IsRunningCommandlet())
 	{
 		FLevelEditorModule& LevelEditorModule = FModuleManager::GetModuleChecked<FLevelEditorModule>(TEXT("LevelEditor"));
-		TSharedPtr<FTabManager> LevelEditorTabManager = LevelEditorModule.GetLevelEditorTabManager();
-		TSharedPtr<SDockTab> FoundTab = LevelEditorTabManager->FindExistingLiveTab(NewTabID);
-		if (FoundTab)
+		if (TSharedPtr<FTabManager> LevelEditorTabManager = LevelEditorModule.GetLevelEditorTabManager())
 		{
-			FoundTab->RequestCloseTab();
-			return true;
+			TSharedPtr<SDockTab> FoundTab = LevelEditorTabManager->FindExistingLiveTab(NewTabID);
+			if (FoundTab)
+			{
+				FoundTab->RequestCloseTab();
+				return true;
+			}
 		}
 	}
 
