AssetDefinition - Further work, adding Materials and associated types as AssetDefinitions, removing their AssetTypeActions.

AssetActionUtility Blueprints - You can now add a further refinement step in addition to a class filter you can specify Content Browser style filter strings and an error message so that people can see the scripted action, but will know why it's not available to them.  Like maybe you have a thing that only can run on non virtual texture streaming assets, you could filter on that tag inaddition to the texture class.

diff --git a/Engine/Plugins/Editor/EngineAssetDefinitions/Source/Private/AssetDefinition_Blueprint.cpp b/Engine/Plugins/Editor/EngineAssetDefinitions/Source/Private/AssetDefinition_Blueprint.cpp
--- a/Engine/Plugins/Editor/EngineAssetDefinitions/Source/Private/AssetDefinition_Blueprint.cpp
+++ b/Engine/Plugins/Editor/EngineAssetDefinitions/Source/Private/AssetDefinition_Blueprint.cpp
@@ -17,6 +17,9 @@
 
 #define LOCTEXT_NAMESPACE "UAssetDefinition_Blueprint"
 
+// Menu Extensions
+//--------------------------------------------------------------------
+
 namespace MenuExtension_Blueprint
 {
 	static bool CanExecuteNewDerivedBlueprint(const FToolMenuContext& MenuContext, const FAssetData* SelectedBlueprintPtr)
diff --git a/Engine/Plugins/Editor/EngineAssetDefinitions/Source/Private/AssetDefinition_MaterialParameterCollection.h b/Engine/Plugins/Editor/EngineAssetDefinitions/Source/Private/AssetDefinition_MaterialParameterCollection.h
deleted file mode 100644
diff --git a/Engine/Plugins/Editor/EngineAssetDefinitions/Source/Private/Material/AssetDefinition_Material.cpp b/Engine/Plugins/Editor/EngineAssetDefinitions/Source/Private/Material/AssetDefinition_Material.cpp
new file mode 100644
--- /dev/null
+++ b/Engine/Plugins/Editor/EngineAssetDefinitions/Source/Private/Material/AssetDefinition_Material.cpp
@@ -0,0 +1,15 @@
+// Copyright Epic Games, Inc. All Rights Reserved.
+
+#include "AssetDefinition_Material.h"
+#include "MaterialEditorModule.h"
+
+EAssetCommandResult UAssetDefinition_Material::OpenAssets(const FAssetOpenArgs& OpenArgs) const
+{
+	for (UMaterial* Material : OpenArgs.LoadObjects<UMaterial>())
+	{
+		IMaterialEditorModule* MaterialEditorModule = &FModuleManager::LoadModuleChecked<IMaterialEditorModule>( "MaterialEditor" );
+		MaterialEditorModule->CreateMaterialEditor(OpenArgs.GetToolkitMode(), OpenArgs.ToolkitHost, Material);
+	}
+
+	return EAssetCommandResult::Handled;
+}
diff --git a/Engine/Plugins/Editor/EngineAssetDefinitions/Source/Private/Material/AssetDefinition_Material.h b/Engine/Plugins/Editor/EngineAssetDefinitions/Source/Private/Material/AssetDefinition_Material.h
new file mode 100644
--- /dev/null
+++ b/Engine/Plugins/Editor/EngineAssetDefinitions/Source/Private/Material/AssetDefinition_Material.h
@@ -0,0 +1,27 @@
+// Copyright Epic Games, Inc. All Rights Reserved.
+
+#pragma once
+
+#include "CoreMinimal.h"
+#include "Materials/Material.h"
+#include "AssetDefinition_MaterialInterface.h"
+
+#include "AssetDefinition_Material.generated.h"
+
+UCLASS()
+class UAssetDefinition_Material : public UAssetDefinition_MaterialInterface
+{
+	GENERATED_BODY()
+
+public:
+	// UAssetDefinition Begin
+	virtual FText GetAssetDisplayName() const override { return NSLOCTEXT("AssetTypeActions", "AssetTypeActions_Material", "Material"); }
+	virtual TSoftClassPtr<UObject> GetAssetClass() const override { return UMaterial::StaticClass(); }
+	virtual TConstArrayView<FAssetCategoryPath> GetAssetCategories() const override
+	{
+		static const auto Categories = { EAssetCategoryPaths::Material, EAssetCategoryPaths::Basic };
+		return Categories;
+	}
+	virtual EAssetCommandResult OpenAssets(const FAssetOpenArgs& OpenArgs) const override;
+	// UAssetDefinition End
+};
diff --git a/Engine/Plugins/Editor/EngineAssetDefinitions/Source/Private/Material/AssetDefinition_MaterialFunction.cpp b/Engine/Plugins/Editor/EngineAssetDefinitions/Source/Private/Material/AssetDefinition_MaterialFunction.cpp
new file mode 100644
--- /dev/null
+++ b/Engine/Plugins/Editor/EngineAssetDefinitions/Source/Private/Material/AssetDefinition_MaterialFunction.cpp
@@ -0,0 +1,217 @@
+// Copyright Epic Games, Inc. All Rights Reserved.
+
+#include "AssetDefinition_MaterialFunction.h"
+
+#include "ContentBrowserMenuContexts.h"
+#include "Factories/MaterialFunctionInstanceFactory.h"
+#include "ToolMenus.h"
+#include "UObject/UObjectHash.h"
+#include "UObject/UObjectIterator.h"
+#include "Styling/AppStyle.h"
+#include "Materials/Material.h"
+#include "Materials/MaterialFunctionInstance.h"
+#include "ThumbnailRendering/SceneThumbnailInfoWithPrimitive.h"
+#include "IAssetTools.h"
+#include "MaterialEditorModule.h"
+#include "ContentBrowserModule.h"
+#include "IContentBrowserSingleton.h"
+#include "AssetRegistry/ARFilter.h"
+#include "AssetRegistry/AssetData.h"
+#include "Misc/PackageName.h"
+
+#define LOCTEXT_NAMESPACE "AssetTypeActions"
+
+const FName MaterialFunctionClass = FName(TEXT("MaterialFunction"));
+const FName MaterialFunctionInstanceClass = FName(TEXT("MaterialFunctionInstance"));
+const FName MaterialFunctionUsageTag = FName(TEXT("MaterialFunctionUsage"));
+const FString LayerCompareString = (TEXT("MaterialLayer"));
+const FString BlendCompareString = (TEXT("MaterialLayerBlend"));
+
+EAssetCommandResult UAssetDefinition_MaterialFunction::OpenAssets(const FAssetOpenArgs& OpenArgs) const
+{
+	for (UMaterialFunction* Function : OpenArgs.LoadObjects<UMaterialFunction>())
+	{
+		IMaterialEditorModule* MaterialEditorModule = &FModuleManager::LoadModuleChecked<IMaterialEditorModule>( "MaterialEditor" );
+		MaterialEditorModule->CreateMaterialEditor(OpenArgs.GetToolkitMode(), OpenArgs.ToolkitHost, Function);
+	}
+
+	return EAssetCommandResult::Handled;
+}
+
+UThumbnailInfo* UAssetDefinition_MaterialFunction::LoadThumbnailInfo(const FAssetData& InAsset) const
+{
+	UMaterialFunctionInterface* MaterialFunc = CastChecked<UMaterialFunctionInterface>(InAsset.GetAsset());
+	UThumbnailInfo* ThumbnailInfo = MaterialFunc->ThumbnailInfo;
+	if (ThumbnailInfo == nullptr)
+	{
+		ThumbnailInfo = NewObject<USceneThumbnailInfoWithPrimitive>(MaterialFunc, NAME_None, RF_Transactional);
+		MaterialFunc->ThumbnailInfo = ThumbnailInfo;
+	}
+
+	return ThumbnailInfo;
+}
+
+// Menu Extensions
+//--------------------------------------------------------------------
+
+namespace MenuExtension_MaterialFunction
+{
+	void ExecuteFindMaterials(const FToolMenuContext& MenuContext)
+	{
+		const UContentBrowserAssetContextMenuContext* CBContext = UContentBrowserAssetContextMenuContext::FindContextWithAssets(MenuContext);
+		
+		if (FEditorDelegates::OnOpenReferenceViewer.IsBound())
+		{
+			// TArray that will be send to the ReferenceViewer for display
+			TArray<FAssetIdentifier> AssetIdentifiers;
+
+			// Iterate over all selected UMaterialFunctionInterface instances
+			for (UMaterialFunctionInterface* MaterialFunction : CBContext->LoadSelectedObjects<UMaterialFunctionInterface>())
+			{
+				// Construct FAssetIdentifier and add to TArray
+				const FName AssetName(*FPackageName::ObjectPathToPackageName(GetPathNameSafe(MaterialFunction)));
+				AssetIdentifiers.Add(FAssetIdentifier(AssetName));
+			}
+			
+			// Call ReferenceViewer
+			FReferenceViewerParams ReferenceViewerParams;
+			ReferenceViewerParams.bShowDependencies = false;
+			ReferenceViewerParams.FixAndHideSearchDepthLimit = 1;
+			ReferenceViewerParams.bShowShowReferencesOptions = false;
+			ReferenceViewerParams.bShowShowSearchableNames = false;
+			FEditorDelegates::OnOpenReferenceViewer.Broadcast(AssetIdentifiers, ReferenceViewerParams);
+		}
+	}
+
+	void ExecuteNewMFI(const FToolMenuContext& MenuContext)
+	{
+		const UContentBrowserAssetContextMenuContext* CBContext = UContentBrowserAssetContextMenuContext::FindContextWithAssets(MenuContext);
+
+		IAssetTools::Get().CreateAssetsFrom<UMaterialFunctionInterface>(
+			CBContext->LoadSelectedObjects<UMaterialFunctionInterface>(), UMaterialFunctionInstance::StaticClass(), TEXT("_Inst"), [](UMaterialFunctionInterface* SourceObject)
+			{
+				UMaterialFunctionInstanceFactory* Factory = NewObject<UMaterialFunctionInstanceFactory>();
+				Factory->InitialParent = SourceObject;
+				return Factory;
+			}
+		);
+	}
+	
+	FDelayedAutoRegisterHelper DelayedAutoRegister(EDelayedRegisterRunPhase::EndOfEngineInit, []{ 
+		UToolMenus::RegisterStartupCallback(FSimpleMulticastDelegate::FDelegate::CreateLambda([]()
+		{
+			FToolMenuOwnerScoped OwnerScoped(UE_MODULE_NAME);
+			UToolMenu* Menu = UE::ContentBrowser::ExtendToolMenu_AssetContextMenu(UMaterialFunction::StaticClass());
+	        
+			FToolMenuSection& Section = Menu->FindOrAddSection("GetAssetActions");
+			Section.AddDynamicEntry("GetAssetActions_NewMaterialFunctionInstance", FNewToolMenuSectionDelegate::CreateLambda([](FToolMenuSection& InSection)
+			{
+				{
+					const TAttribute<FText> Label = NSLOCTEXT("AssetTypeActions", "Material_NewMFI", "Create Function Instance");
+					const TAttribute<FText> ToolTip = NSLOCTEXT("AssetTypeActions", "Material_NewMFITooltip", "Creates a parameterized function using this function as a base.");
+					const FSlateIcon Icon = FSlateIcon(FAppStyle::GetAppStyleSetName(), "ClassIcon.MaterialInstanceActor");
+					const FToolMenuExecuteAction UIAction = FToolMenuExecuteAction::CreateStatic(&ExecuteNewMFI);
+
+					InSection.AddMenuEntry("MaterialFunction_NewMFI", Label, ToolTip, Icon, UIAction);
+				}
+			}));
+			
+			Section.AddDynamicEntry("GetAssetActions_UMaterialFunction", FNewToolMenuSectionDelegate::CreateLambda([](FToolMenuSection& InSection)
+			{
+				if (FEditorDelegates::OnOpenReferenceViewer.IsBound())
+				{
+					const TAttribute<FText> Label = LOCTEXT("MaterialFunction_FindMaterials", "Find Materials Using This");
+					const TAttribute<FText> ToolTip = LOCTEXT("MaterialFunction_FindMaterialsTooltip", "Finds the materials that reference this material function and visually displays them with the light version of the Reference Viewer.");
+					const FSlateIcon Icon = FSlateIcon(FAppStyle::GetAppStyleSetName(), "ContentBrowser.AssetActions.GenericFind");
+					const FToolMenuExecuteAction UIAction = FToolMenuExecuteAction::CreateStatic(&ExecuteFindMaterials);
+
+					InSection.AddMenuEntry("MaterialFunction_FindMaterials", Label, ToolTip, Icon, UIAction);
+				}
+			}));
+		}));
+	});
+}
+
+namespace MenuExtension_MaterialFunctionMaterialLayer
+{
+	void ExecuteNewMFI(const FToolMenuContext& MenuContext)
+	{
+		const UContentBrowserAssetContextMenuContext* CBContext = UContentBrowserAssetContextMenuContext::FindContextWithAssets(MenuContext);
+
+		IAssetTools::Get().CreateAssetsFrom<UMaterialFunctionInterface>(
+			CBContext->LoadSelectedObjects<UMaterialFunctionInterface>(), UMaterialFunctionMaterialLayerInstance::StaticClass(), TEXT("_Inst"), [](UMaterialFunctionInterface* SourceObject)
+			{
+				UMaterialFunctionMaterialLayerInstanceFactory* Factory = NewObject<UMaterialFunctionMaterialLayerInstanceFactory>();
+				Factory->InitialParent = SourceObject;
+				return Factory;
+			}
+		);
+	}
+	
+	FDelayedAutoRegisterHelper DelayedAutoRegister(EDelayedRegisterRunPhase::EndOfEngineInit, []{ 
+		UToolMenus::RegisterStartupCallback(FSimpleMulticastDelegate::FDelegate::CreateLambda([]()
+		{
+			FToolMenuOwnerScoped OwnerScoped(UE_MODULE_NAME);
+			TArray<UToolMenu*> Menus = UE::ContentBrowser::ExtendToolMenu_AssetContextMenus({ UMaterialFunctionMaterialLayer::StaticClass(), UMaterialFunctionMaterialLayerInstance::StaticClass() });
+			
+			for (UToolMenu* Menu : Menus)
+			{
+				FToolMenuSection& Section = Menu->FindOrAddSection("GetAssetActions");
+				Section.AddDynamicEntry("GetAssetActions_NewMaterialFunctionInstance", FNewToolMenuSectionDelegate::CreateLambda([](FToolMenuSection& InSection)
+				{
+					{
+						const TAttribute<FText> Label = NSLOCTEXT("AssetTypeActions", "Material_NewMLI", "Create Layer Instance");
+						const TAttribute<FText> ToolTip = NSLOCTEXT("AssetTypeActions", "Material_NewMFITooltip", "Creates a parameterized function using this function as a base.");
+						const FSlateIcon Icon = FSlateIcon(FAppStyle::GetAppStyleSetName(), "ClassIcon.MaterialInstanceActor");
+						const FToolMenuExecuteAction UIAction = FToolMenuExecuteAction::CreateStatic(&ExecuteNewMFI);
+
+						InSection.AddMenuEntry("MaterialFunction_NewMFI", Label, ToolTip, Icon, UIAction);
+					}
+				}));
+			}
+		}));
+	});
+}
+
+namespace MenuExtension_MaterialFunctionMaterialLayerBlend
+{
+	void ExecuteNewMFI(const FToolMenuContext& MenuContext)
+	{
+		const UContentBrowserAssetContextMenuContext* CBContext = UContentBrowserAssetContextMenuContext::FindContextWithAssets(MenuContext);
+
+		IAssetTools::Get().CreateAssetsFrom<UMaterialFunctionInterface>(
+			CBContext->LoadSelectedObjects<UMaterialFunctionInterface>(), UMaterialFunctionMaterialLayerBlendInstance::StaticClass(), TEXT("_Inst"), [](UMaterialFunctionInterface* SourceObject)
+			{
+				UMaterialFunctionMaterialLayerBlendInstanceFactory* Factory = NewObject<UMaterialFunctionMaterialLayerBlendInstanceFactory>();
+				Factory->InitialParent = SourceObject;
+				return Factory;
+			}
+		);
+	}
+	
+	FDelayedAutoRegisterHelper DelayedAutoRegister(EDelayedRegisterRunPhase::EndOfEngineInit, []{ 
+		UToolMenus::RegisterStartupCallback(FSimpleMulticastDelegate::FDelegate::CreateLambda([]()
+		{
+			FToolMenuOwnerScoped OwnerScoped(UE_MODULE_NAME);
+			TArray<UToolMenu*> Menus = UE::ContentBrowser::ExtendToolMenu_AssetContextMenus({ UMaterialFunctionMaterialLayerBlend::StaticClass(), UMaterialFunctionMaterialLayerBlendInstance::StaticClass() });
+
+			for (UToolMenu* Menu : Menus)
+			{
+				FToolMenuSection& Section = Menu->FindOrAddSection("GetAssetActions");
+				Section.AddDynamicEntry("GetAssetActions_NewMaterialFunctionInstance", FNewToolMenuSectionDelegate::CreateLambda([](FToolMenuSection& InSection)
+				{
+					{
+						const TAttribute<FText> Label = NSLOCTEXT("AssetTypeActions", "Material_NewMBI", "Create Blend Instance");
+						const TAttribute<FText> ToolTip = NSLOCTEXT("AssetTypeActions", "Material_NewMFITooltip", "Creates a parameterized function using this function as a base.");
+						const FSlateIcon Icon = FSlateIcon(FAppStyle::GetAppStyleSetName(), "ClassIcon.MaterialInstanceActor");
+						const FToolMenuExecuteAction UIAction = FToolMenuExecuteAction::CreateStatic(&ExecuteNewMFI);
+
+						InSection.AddMenuEntry("MaterialFunction_NewMFI", Label, ToolTip, Icon, UIAction);
+					}
+				}));
+			}
+		}));
+	});
+}
+
+#undef LOCTEXT_NAMESPACE
diff --git a/Engine/Plugins/Editor/EngineAssetDefinitions/Source/Private/Material/AssetDefinition_MaterialFunction.h b/Engine/Plugins/Editor/EngineAssetDefinitions/Source/Private/Material/AssetDefinition_MaterialFunction.h
new file mode 100644
--- /dev/null
+++ b/Engine/Plugins/Editor/EngineAssetDefinitions/Source/Private/Material/AssetDefinition_MaterialFunction.h
@@ -0,0 +1,55 @@
+// Copyright Epic Games, Inc. All Rights Reserved.
+
+#pragma once
+
+#include "CoreMinimal.h"
+#include "Materials/MaterialFunction.h"
+#include "Materials/MaterialFunctionMaterialLayer.h"
+#include "Materials/MaterialFunctionMaterialLayerBlend.h"
+#include "AssetDefinitionDefault.h"
+
+#include "AssetDefinition_MaterialFunction.generated.h"
+
+UCLASS()
+class UAssetDefinition_MaterialFunction : public UAssetDefinitionDefault
+{
+	GENERATED_BODY()
+
+public:
+	// UAssetDefinition Begin
+	virtual FText GetAssetDisplayName() const override { return NSLOCTEXT("AssetTypeActions", "AssetTypeActions_MaterialFunction", "Material Function"); }
+	virtual FLinearColor GetAssetColor() const override { return FLinearColor(FColor(0,175,175)); }
+	virtual TSoftClassPtr<UObject> GetAssetClass() const override { return UMaterialFunction::StaticClass(); }
+	virtual TConstArrayView<FAssetCategoryPath> GetAssetCategories() const override
+	{
+		static const auto Categories = { EAssetCategoryPaths::Material };
+		return Categories;
+	}
+	virtual UThumbnailInfo* LoadThumbnailInfo(const FAssetData& InAsset) const override;
+	virtual EAssetCommandResult OpenAssets(const FAssetOpenArgs& OpenArgs) const override;
+	// UAssetDefinition End
+};
+
+UCLASS()
+class UAssetDefinition_MaterialFunctionMaterialLayer : public UAssetDefinition_MaterialFunction
+{
+	GENERATED_BODY()
+
+public:
+	// UAssetDefinition Begin
+	virtual FText GetAssetDisplayName() const override { return NSLOCTEXT("AssetTypeActions", "AssetTypeActions_MaterialFunctionMaterialLayer", "Material Layer"); }
+	virtual TSoftClassPtr<UObject> GetAssetClass() const override { return UMaterialFunctionMaterialLayer::StaticClass(); }
+	// UAssetDefinition End
+};
+
+UCLASS()
+class UAssetDefinition_MaterialFunctionLayerBlend : public UAssetDefinition_MaterialFunction
+{
+	GENERATED_BODY()
+
+public:
+	// UAssetDefinition Begin
+	virtual FText GetAssetDisplayName() const override { return NSLOCTEXT("AssetTypeActions", "AssetTypeActions_MaterialFunctionMaterialLayerBlend", "Material Layer Blend"); }
+	virtual TSoftClassPtr<UObject> GetAssetClass() const override { return UMaterialFunctionMaterialLayerBlend::StaticClass(); }
+	// UAssetDefinition End
+};
diff --git a/Engine/Plugins/Editor/EngineAssetDefinitions/Source/Private/Material/AssetDefinition_MaterialFunctionInstance.cpp b/Engine/Plugins/Editor/EngineAssetDefinitions/Source/Private/Material/AssetDefinition_MaterialFunctionInstance.cpp
new file mode 100644
--- /dev/null
+++ b/Engine/Plugins/Editor/EngineAssetDefinitions/Source/Private/Material/AssetDefinition_MaterialFunctionInstance.cpp
@@ -0,0 +1,90 @@
+// Copyright Epic Games, Inc. All Rights Reserved.
+
+#include "AssetDefinition_MaterialFunctionInstance.h"
+
+#include "ContentBrowserMenuContexts.h"
+#include "ToolMenus.h"
+#include "UObject/UObjectHash.h"
+#include "UObject/UObjectIterator.h"
+#include "Styling/AppStyle.h"
+#include "ThumbnailRendering/SceneThumbnailInfoWithPrimitive.h"
+#include "IAssetTools.h"
+#include "MaterialEditorModule.h"
+#include "Factories/MaterialFunctionInstanceFactory.h"
+#include "ContentBrowserModule.h"
+#include "IContentBrowserSingleton.h"
+
+#define LOCTEXT_NAMESPACE "AssetTypeActions"
+
+EAssetCommandResult UAssetDefinition_MaterialFunctionInstance::OpenAssets(const FAssetOpenArgs& OpenArgs) const
+{
+	for (UMaterialFunctionInstance* MFI : OpenArgs.LoadObjects<UMaterialFunctionInstance>())
+	{
+		IMaterialEditorModule* MaterialEditorModule = &FModuleManager::LoadModuleChecked<IMaterialEditorModule>( "MaterialEditor" );
+		MaterialEditorModule->CreateMaterialInstanceEditor(OpenArgs.GetToolkitMode(), OpenArgs.ToolkitHost, MFI);
+	}
+
+	return EAssetCommandResult::Handled;
+}
+
+UThumbnailInfo* UAssetDefinition_MaterialFunctionInstance::LoadThumbnailInfo(const FAssetData& InAsset) const
+{
+	UMaterialFunctionInstance* MaterialFunc = CastChecked<UMaterialFunctionInstance>(InAsset.GetAsset());
+	UThumbnailInfo* ThumbnailInfo = MaterialFunc->ThumbnailInfo;
+	if (ThumbnailInfo == nullptr)
+	{
+		ThumbnailInfo = NewObject<USceneThumbnailInfoWithPrimitive>(MaterialFunc, NAME_None, RF_Transactional);
+		MaterialFunc->ThumbnailInfo = ThumbnailInfo;
+	}
+
+	return ThumbnailInfo;
+}
+
+// Menu Extensions
+//--------------------------------------------------------------------
+
+namespace MenuExtension_MaterialFunctionInstance
+{
+	void ExecuteFindParent(const FToolMenuContext& MenuContext)
+	{
+		const UContentBrowserAssetContextMenuContext* CBContext = UContentBrowserAssetContextMenuContext::FindContextWithAssets(MenuContext);
+
+		TArray<UObject*> ObjectsToSyncTo;
+		for (UMaterialFunctionInstance* MFI : CBContext->LoadSelectedObjects<UMaterialFunctionInstance>())
+		{
+			if ( MFI->Parent )
+			{
+				ObjectsToSyncTo.AddUnique( MFI->Parent );
+			}
+		}
+
+		// Sync the respective browser to the valid parents
+		if ( ObjectsToSyncTo.Num() > 0 )
+		{
+			IAssetTools::Get().SyncBrowserToAssets(ObjectsToSyncTo);
+		}
+	}
+	
+	FDelayedAutoRegisterHelper DelayedAutoRegister(EDelayedRegisterRunPhase::EndOfEngineInit, []{ 
+		UToolMenus::RegisterStartupCallback(FSimpleMulticastDelegate::FDelegate::CreateLambda([]()
+		{
+			FToolMenuOwnerScoped OwnerScoped(UE_MODULE_NAME);
+			UToolMenu* Menu = UE::ContentBrowser::ExtendToolMenu_AssetContextMenu(UMaterialFunctionInstance::StaticClass());
+	        
+			FToolMenuSection& Section = Menu->FindOrAddSection("GetAssetActions");
+			Section.AddDynamicEntry("GetAssetActions_UMaterialFunctionInstance", FNewToolMenuSectionDelegate::CreateLambda([](FToolMenuSection& InSection)
+			{
+				{
+					const TAttribute<FText> Label = LOCTEXT("MaterialFunctionInstance_FindParent", "Find Parent");
+					const TAttribute<FText> ToolTip = LOCTEXT("MaterialFunctionInstance_FindParentTooltip", "Finds the function this instance is based on in the content browser.");
+					const FSlateIcon Icon = FSlateIcon(FAppStyle::GetAppStyleSetName(), "ContentBrowser.AssetActions.GenericFind");
+					const FToolMenuExecuteAction UIAction = FToolMenuExecuteAction::CreateStatic(&ExecuteFindParent);
+
+					InSection.AddMenuEntry("MaterialFunctionInstance_FindParent", Label, ToolTip, Icon, UIAction);
+				}
+			}));
+		}));
+	});
+}
+
+#undef LOCTEXT_NAMESPACE
diff --git a/Engine/Plugins/Editor/EngineAssetDefinitions/Source/Private/Material/AssetDefinition_MaterialFunctionInstance.h b/Engine/Plugins/Editor/EngineAssetDefinitions/Source/Private/Material/AssetDefinition_MaterialFunctionInstance.h
new file mode 100644
--- /dev/null
+++ b/Engine/Plugins/Editor/EngineAssetDefinitions/Source/Private/Material/AssetDefinition_MaterialFunctionInstance.h
@@ -0,0 +1,47 @@
+// Copyright Epic Games, Inc. All Rights Reserved.
+
+#pragma once
+
+#include "CoreMinimal.h"
+#include "AssetDefinition_MaterialFunction.h"
+#include "Materials/MaterialFunctionInstance.h"
+
+#include "AssetDefinition_MaterialFunctionInstance.generated.h"
+
+UCLASS()
+class UAssetDefinition_MaterialFunctionInstance : public UAssetDefinition_MaterialFunction
+{
+	GENERATED_BODY()
+
+public:
+	// UAssetDefinition Begin
+	virtual FText GetAssetDisplayName() const override { return NSLOCTEXT("AssetTypeActions", "FAssetTypeActions_MaterialFunctionInstance", "Material Function Instance"); }
+	virtual TSoftClassPtr<UObject> GetAssetClass() const override { return UMaterialFunctionInstance::StaticClass(); }
+	virtual UThumbnailInfo* LoadThumbnailInfo(const FAssetData& InAsset) const override;
+	virtual EAssetCommandResult OpenAssets(const FAssetOpenArgs& OpenArgs) const override;
+	// UAssetDefinition End
+};
+
+UCLASS()
+class UAssetDefinition_MaterialFunctionLayerInstance : public UAssetDefinition_MaterialFunctionInstance
+{
+	GENERATED_BODY()
+
+public:
+	// UAssetDefinition Begin
+	virtual FText GetAssetDisplayName() const override { return NSLOCTEXT("AssetTypeActions", "AssetTypeActions_MaterialFunctionMaterialLayerInstance", "Material Layer Instance"); }
+	virtual TSoftClassPtr<UObject> GetAssetClass() const override { return UMaterialFunctionMaterialLayerInstance::StaticClass(); }
+	// UAssetDefinition End
+};
+
+UCLASS()
+class UAssetDefinition_MaterialFunctionLayerBlendInstance : public UAssetDefinition_MaterialFunctionInstance
+{
+	GENERATED_BODY()
+
+public:
+	// UAssetDefinition Begin
+	virtual FText GetAssetDisplayName() const override { return NSLOCTEXT("AssetTypeActions", "AssetTypeActions_MaterialFunctionMaterialLayerBlendInstance", "Material Layer Blend Instance"); }
+	virtual TSoftClassPtr<UObject> GetAssetClass() const override { return UMaterialFunctionMaterialLayerBlendInstance::StaticClass(); }
+	// UAssetDefinition End
+};
diff --git a/Engine/Plugins/Editor/EngineAssetDefinitions/Source/Private/Material/AssetDefinition_MaterialInstanceConstant.cpp b/Engine/Plugins/Editor/EngineAssetDefinitions/Source/Private/Material/AssetDefinition_MaterialInstanceConstant.cpp
new file mode 100644
--- /dev/null
+++ b/Engine/Plugins/Editor/EngineAssetDefinitions/Source/Private/Material/AssetDefinition_MaterialInstanceConstant.cpp
@@ -0,0 +1,73 @@
+// Copyright Epic Games, Inc. All Rights Reserved.
+
+#include "AssetDefinition_MaterialInstanceConstant.h"
+
+#include "ContentBrowserMenuContexts.h"
+#include "ContentBrowserModule.h"
+#include "ToolMenus.h"
+#include "Styling/AppStyle.h"
+#include "IAssetTools.h"
+#include "IContentBrowserSingleton.h"
+#include "MaterialEditorModule.h"
+
+#define LOCTEXT_NAMESPACE "AssetTypeActions"
+
+EAssetCommandResult UAssetDefinition_MaterialInstanceConstant::OpenAssets(const FAssetOpenArgs& OpenArgs) const
+{
+	for (UMaterialInstanceConstant* MIC : OpenArgs.LoadObjects<UMaterialInstanceConstant>())
+   	{
+		IMaterialEditorModule* MaterialEditorModule = &FModuleManager::LoadModuleChecked<IMaterialEditorModule>( "MaterialEditor" );
+		MaterialEditorModule->CreateMaterialInstanceEditor(OpenArgs.GetToolkitMode(), OpenArgs.ToolkitHost, MIC);
+	}
+
+	return EAssetCommandResult::Handled;
+}
+
+// Menu Extensions
+//--------------------------------------------------------------------
+
+namespace
+{
+	void ExecuteFindParent(const FToolMenuContext& MenuContext)
+	{
+		const UContentBrowserAssetContextMenuContext* CBContext = UContentBrowserAssetContextMenuContext::FindContextWithAssets(MenuContext);
+
+		TArray<UObject*> ObjectsToSyncTo;
+		for (UMaterialInstanceConstant* MIC : CBContext->LoadSelectedObjects<UMaterialInstanceConstant>())
+		{
+			if ( MIC->Parent )
+			{
+				ObjectsToSyncTo.AddUnique( MIC->Parent );
+			}
+		}
+
+		// Sync the respective browser to the valid parents
+		if ( ObjectsToSyncTo.Num() > 0 )
+		{
+			IAssetTools::Get().SyncBrowserToAssets(ObjectsToSyncTo);
+		}
+	}
+	
+	FDelayedAutoRegisterHelper DelayedAutoRegister(EDelayedRegisterRunPhase::EndOfEngineInit, []{ 
+		UToolMenus::RegisterStartupCallback(FSimpleMulticastDelegate::FDelegate::CreateLambda([]()
+		{
+			FToolMenuOwnerScoped OwnerScoped(UE_MODULE_NAME);
+			UToolMenu* Menu = UE::ContentBrowser::ExtendToolMenu_AssetContextMenu(UMaterialInstanceConstant::StaticClass());
+	        
+			FToolMenuSection& Section = Menu->FindOrAddSection("GetAssetActions");
+			Section.AddDynamicEntry("GetAssetActions_MaterialInstanceConstant", FNewToolMenuSectionDelegate::CreateLambda([](FToolMenuSection& InSection)
+			{
+				{
+					const TAttribute<FText> Label = LOCTEXT("MaterialInstanceConstant_FindParent", "Find Parent");
+					const TAttribute<FText> ToolTip = LOCTEXT("MaterialInstanceConstant_FindParentTooltip", "Finds the material this instance is based on in the content browser.");
+					const FSlateIcon Icon = FSlateIcon(FAppStyle::GetAppStyleSetName(), "ContentBrowser.AssetActions.GenericFind");
+					const FToolMenuExecuteAction UIAction = FToolMenuExecuteAction::CreateStatic(&ExecuteFindParent);
+
+					InSection.AddMenuEntry("MaterialInstanceConstant_FindParent", Label, ToolTip, Icon, UIAction);
+				}
+			}));
+		}));
+	});
+}
+
+#undef LOCTEXT_NAMESPACE
diff --git a/Engine/Plugins/Editor/EngineAssetDefinitions/Source/Private/Material/AssetDefinition_MaterialInstanceConstant.h b/Engine/Plugins/Editor/EngineAssetDefinitions/Source/Private/Material/AssetDefinition_MaterialInstanceConstant.h
new file mode 100644
--- /dev/null
+++ b/Engine/Plugins/Editor/EngineAssetDefinitions/Source/Private/Material/AssetDefinition_MaterialInstanceConstant.h
@@ -0,0 +1,23 @@
+// Copyright Epic Games, Inc. All Rights Reserved.
+
+#pragma once
+
+#include "CoreMinimal.h"
+#include "Materials/MaterialInstanceConstant.h"
+#include "AssetDefinition_MaterialInterface.h"
+
+#include "AssetDefinition_MaterialInstanceConstant.generated.h"
+
+UCLASS()
+class UAssetDefinition_MaterialInstanceConstant : public UAssetDefinition_MaterialInterface
+{
+	GENERATED_BODY()
+
+public:
+	// UAssetDefinition Begin
+	virtual FText GetAssetDisplayName() const override { return NSLOCTEXT("AssetTypeActions", "AssetTypeActions_MaterialInstanceConstant", "Material Instance"); }
+	virtual FLinearColor GetAssetColor() const override { return FLinearColor(FColor(0,128,0)); }
+	virtual TSoftClassPtr<UObject> GetAssetClass() const override { return UMaterialInstanceConstant::StaticClass(); }
+	virtual EAssetCommandResult OpenAssets(const FAssetOpenArgs& OpenArgs) const override;
+	// UAssetDefinition End
+};
diff --git a/Engine/Plugins/Editor/EngineAssetDefinitions/Source/Private/Material/AssetDefinition_MaterialInstanceDynamic.cpp b/Engine/Plugins/Editor/EngineAssetDefinitions/Source/Private/Material/AssetDefinition_MaterialInstanceDynamic.cpp
new file mode 100644
--- /dev/null
+++ b/Engine/Plugins/Editor/EngineAssetDefinitions/Source/Private/Material/AssetDefinition_MaterialInstanceDynamic.cpp
@@ -0,0 +1,19 @@
+// Copyright Epic Games, Inc. All Rights Reserved.
+
+#include "AssetDefinition_MaterialInstanceDynamic.h"
+#include "Toolkits/SimpleAssetEditor.h"
+
+#define LOCTEXT_NAMESPACE "AssetTypeActions"
+
+EAssetCommandResult UAssetDefinition_MaterialInstanceDynamic::OpenAssets(const FAssetOpenArgs& OpenArgs) const
+{
+	TSharedRef<FSimpleAssetEditor> Editor = FSimpleAssetEditor::CreateEditor(EToolkitMode::Standalone, OpenArgs.ToolkitHost, OpenArgs.LoadObjects<UObject>());
+	Editor->SetPropertyEditingEnabledDelegate(FIsPropertyEditingEnabled::CreateStatic([]() -> bool
+	{
+		return true;
+	}));
+
+	return EAssetCommandResult::Handled;
+}
+
+#undef LOCTEXT_NAMESPACE
diff --git a/Engine/Plugins/Editor/EngineAssetDefinitions/Source/Private/Material/AssetDefinition_MaterialInstanceDynamic.h b/Engine/Plugins/Editor/EngineAssetDefinitions/Source/Private/Material/AssetDefinition_MaterialInstanceDynamic.h
new file mode 100644
--- /dev/null
+++ b/Engine/Plugins/Editor/EngineAssetDefinitions/Source/Private/Material/AssetDefinition_MaterialInstanceDynamic.h
@@ -0,0 +1,22 @@
+// Copyright Epic Games, Inc. All Rights Reserved.
+
+#pragma once
+
+#include "CoreMinimal.h"
+#include "Materials/MaterialInstanceDynamic.h"
+#include "AssetDefinition_MaterialInterface.h"
+
+#include "AssetDefinition_MaterialInstanceDynamic.generated.h"
+
+UCLASS()
+class UAssetDefinition_MaterialInstanceDynamic : public UAssetDefinition_MaterialInterface
+{
+	GENERATED_BODY()
+
+public:
+	// UAssetDefinition Begin
+	virtual FText GetAssetDisplayName() const override { return NSLOCTEXT("AssetTypeActions", "AssetTypeActions_MaterialInstanceDynamic", "Material Instance Dynamic"); }
+	virtual TSoftClassPtr<UObject> GetAssetClass() const override { return UMaterialInstanceDynamic::StaticClass(); }
+	virtual EAssetCommandResult OpenAssets(const FAssetOpenArgs& OpenArgs) const override;
+	// UAssetDefinition End
+};
diff --git a/Engine/Plugins/Editor/EngineAssetDefinitions/Source/Private/Material/AssetDefinition_MaterialInterface.cpp b/Engine/Plugins/Editor/EngineAssetDefinitions/Source/Private/Material/AssetDefinition_MaterialInterface.cpp
new file mode 100644
--- /dev/null
+++ b/Engine/Plugins/Editor/EngineAssetDefinitions/Source/Private/Material/AssetDefinition_MaterialInterface.cpp
@@ -0,0 +1,79 @@
+// Copyright Epic Games, Inc. All Rights Reserved.
+
+#include "Material/AssetDefinition_MaterialInterface.h"
+
+#include "ContentBrowserMenuContexts.h"
+#include "Materials/Material.h"
+#include "ToolMenus.h"
+#include "Styling/AppStyle.h"
+#include "Factories/MaterialInstanceConstantFactoryNew.h"
+#include "ThumbnailRendering/SceneThumbnailInfoWithPrimitive.h"
+#include "Materials/MaterialInstanceConstant.h"
+#include "IAssetTools.h"
+#include "IContentBrowserSingleton.h"
+#include "ContentBrowserModule.h"
+
+#define LOCTEXT_NAMESPACE "AssetTypeActions"
+
+UThumbnailInfo* UAssetDefinition_MaterialInterface::LoadThumbnailInfo(const FAssetData& InAsset) const
+{
+	UMaterialInterface* MaterialInterface = CastChecked<UMaterialInterface>(InAsset.GetAsset());
+	USceneThumbnailInfoWithPrimitive* ThumbnailInfo = Cast<USceneThumbnailInfoWithPrimitive>(MaterialInterface->ThumbnailInfo);
+	if ( ThumbnailInfo == nullptr )
+	{
+		ThumbnailInfo = NewObject<USceneThumbnailInfoWithPrimitive>(MaterialInterface, NAME_None, RF_Transactional);
+		MaterialInterface->ThumbnailInfo = ThumbnailInfo;
+	}
+	
+	const UMaterial* Material = MaterialInterface->GetBaseMaterial();
+	if (Material && Material->bUsedWithParticleSprites)
+    {
+    	ThumbnailInfo->DefaultPrimitiveType = TPT_Plane;
+    }
+
+	return ThumbnailInfo;
+}
+
+// Menu Extensions
+//--------------------------------------------------------------------
+
+namespace MenuExtension_MaterialInterface
+{
+	static void ExecuteNewMIC(const FToolMenuContext& MenuContext)
+	{
+		const UContentBrowserAssetContextMenuContext* CBContext = UContentBrowserAssetContextMenuContext::FindContextWithAssets(MenuContext);
+
+		IAssetTools::Get().CreateAssetsFrom<UMaterialInterface>(
+			CBContext->LoadSelectedObjects<UMaterialInterface>(), UMaterialInstanceConstant::StaticClass(), TEXT("_Inst"), [](UMaterialInterface* SourceObject)
+			{
+				UMaterialInstanceConstantFactoryNew* Factory = NewObject<UMaterialInstanceConstantFactoryNew>();
+				Factory->InitialParent = SourceObject;
+				return Factory;
+			}
+		);
+	}
+	
+	static FDelayedAutoRegisterHelper DelayedAutoRegister(EDelayedRegisterRunPhase::EndOfEngineInit, []{ 
+		UToolMenus::RegisterStartupCallback(FSimpleMulticastDelegate::FDelegate::CreateLambda([]()
+		{
+			FToolMenuOwnerScoped OwnerScoped(UE_MODULE_NAME);
+			UToolMenu* Menu = UE::ContentBrowser::ExtendToolMenu_AssetContextMenu(UMaterialInterface::StaticClass());
+	        
+			FToolMenuSection& Section = Menu->FindOrAddSection("GetAssetActions");
+			Section.AddDynamicEntry("GetAssetActions_UMaterialInterface", FNewToolMenuSectionDelegate::CreateLambda([](FToolMenuSection& InSection)
+			{
+				{
+					const TAttribute<FText> Label = LOCTEXT("Material_NewMIC", "Create Material Instance");
+					const TAttribute<FText> ToolTip = LOCTEXT("Material_NewMICTooltip", "Creates a parameterized material using this material as a base.");
+					const FSlateIcon Icon = FSlateIcon(FAppStyle::GetAppStyleSetName(), "ClassIcon.MaterialInstanceActor");
+					const FToolMenuExecuteAction UIAction = FToolMenuExecuteAction::CreateStatic(&ExecuteNewMIC);
+
+					InSection.AddMenuEntry("Material_NewMIC", Label, ToolTip, Icon, UIAction);
+				}
+			}));
+		}));
+	});
+}
+
+
+#undef LOCTEXT_NAMESPACE
diff --git a/Engine/Plugins/Editor/EngineAssetDefinitions/Source/Private/Material/AssetDefinition_MaterialInterface.h b/Engine/Plugins/Editor/EngineAssetDefinitions/Source/Private/Material/AssetDefinition_MaterialInterface.h
new file mode 100644
--- /dev/null
+++ b/Engine/Plugins/Editor/EngineAssetDefinitions/Source/Private/Material/AssetDefinition_MaterialInterface.h
@@ -0,0 +1,28 @@
+// Copyright Epic Games, Inc. All Rights Reserved.
+
+#pragma once
+
+#include "CoreMinimal.h"
+#include "Materials/MaterialInterface.h"
+#include "AssetDefinitionDefault.h"
+
+#include "AssetDefinition_MaterialInterface.generated.h"
+
+UCLASS()
+class UAssetDefinition_MaterialInterface : public UAssetDefinitionDefault
+{
+	GENERATED_BODY()
+
+public:
+	// UAssetDefinition Begin
+	virtual FText GetAssetDisplayName() const override { return NSLOCTEXT("AssetTypeActions", "AssetTypeActions_MaterialInterface", "Material Interface"); }
+	virtual FLinearColor GetAssetColor() const override { return FLinearColor(FColor(64,192,64)); }
+	virtual TSoftClassPtr<UObject> GetAssetClass() const override { return UMaterialInterface::StaticClass(); }
+	virtual TConstArrayView<FAssetCategoryPath> GetAssetCategories() const override
+	{
+		static const auto Categories = { EAssetCategoryPaths::Material };
+		return Categories;
+	}
+	virtual UThumbnailInfo* LoadThumbnailInfo(const FAssetData& InAsset) const override;
+	// UAssetDefinition End
+};
diff --git a/Engine/Plugins/Editor/EngineAssetDefinitions/Source/Private/Material/AssetDefinition_MaterialParameterCollection.h b/Engine/Plugins/Editor/EngineAssetDefinitions/Source/Private/Material/AssetDefinition_MaterialParameterCollection.h
new file mode 100644
--- /dev/null
+++ b/Engine/Plugins/Editor/EngineAssetDefinitions/Source/Private/Material/AssetDefinition_MaterialParameterCollection.h
@@ -0,0 +1,27 @@
+// Copyright Epic Games, Inc. All Rights Reserved.
+
+#pragma once
+
+#include "CoreMinimal.h"
+#include "Materials/MaterialParameterCollection.h"
+#include "AssetDefinitionDefault.h"
+
+#include "AssetDefinition_MaterialParameterCollection.generated.h"
+
+UCLASS()
+class UAssetDefinition_MaterialParameterCollection : public UAssetDefinitionDefault
+{
+	GENERATED_BODY()
+
+public:
+	// UAssetDefinition Begin
+	virtual FText GetAssetDisplayName() const override { return NSLOCTEXT("AssetTypeActions", "AssetTypeActions_MaterialParameterCollection", "Material Parameter Collection"); }
+	virtual FLinearColor GetAssetColor() const override { return FLinearColor(FColor(0, 192, 0)); }
+	virtual TSoftClassPtr<UObject> GetAssetClass() const override { return UMaterialParameterCollection::StaticClass(); }
+	virtual TConstArrayView<FAssetCategoryPath> GetAssetCategories() const override
+	{
+		static const auto Categories = { EAssetCategoryPaths::Material };
+		return Categories;
+	}
+	// UAssetDefinition End
+};
diff --git a/Engine/Source/Developer/AssetTools/Private/AssetTools.cpp b/Engine/Source/Developer/AssetTools/Private/AssetTools.cpp
--- a/Engine/Source/Developer/AssetTools/Private/AssetTools.cpp
+++ b/Engine/Source/Developer/AssetTools/Private/AssetTools.cpp
@@ -37,7 +37,6 @@
 #include "ClassViewerFilter.h"
 #include "AssetTypeActions/AssetTypeActions_Blueprint.h"
 #include "AssetTypeActions/AssetTypeActions_BlueprintGeneratedClass.h"
-#include "AssetTypeActions/AssetTypeActions_MaterialInterface.h"
 #include "AssetTypeActions/AssetTypeActions_SkeletalMesh.h"
 #include "AssetTypeActions/AssetTypeActions_AnimationAsset.h"
 #include "AssetTypeActions/AssetTypeActions_AnimBlueprint.h"
@@ -65,11 +64,6 @@
 #include "AssetTypeActions/AssetTypeActions_HapticFeedback.h"
 #include "AssetTypeActions/AssetTypeActions_ActorFoliageSettings.h"
 #include "AssetTypeActions/AssetTypeActions_LightWeightInstance.h"
-#include "AssetTypeActions/AssetTypeActions_Material.h"
-#include "AssetTypeActions/AssetTypeActions_MaterialFunction.h"
-#include "AssetTypeActions/AssetTypeActions_MaterialFunctionInstance.h"
-#include "AssetTypeActions/AssetTypeActions_MaterialInstanceConstant.h"
-#include "AssetTypeActions/AssetTypeActions_MaterialInstanceDynamic.h"
 #include "AssetTypeActions/AssetTypeActions_MirrorDataTable.h"
 #include "AssetTypeActions/AssetTypeActions_ParticleSystem.h"
 #include "AssetTypeActions/AssetTypeActions_PhysicalMaterialMask.h"
@@ -1268,16 +1262,6 @@ UAssetToolsImpl::UAssetToolsImpl(const FObjectInitializer& ObjectInitializer)
 	RegisterAssetTypeActions(MakeShareable(new FAssetTypeActions_ForceFeedbackEffect(InputCategoryBit)));
 	RegisterAssetTypeActions(MakeShareable(new FAssetTypeActions_ActorFoliageSettings(FoliageCategoryBit)));
 	RegisterAssetTypeActions(MakeShareable(new FAssetTypeActions_LightWeightInstance));
-	RegisterAssetTypeActions(MakeShareable(new FAssetTypeActions_Material(EAssetTypeCategories::Type::None)));
-	RegisterAssetTypeActions(MakeShareable(new FAssetTypeActions_MaterialFunction));
-	RegisterAssetTypeActions(MakeShareable(new FAssetTypeActions_MaterialFunctionLayer));
-	RegisterAssetTypeActions(MakeShareable(new FAssetTypeActions_MaterialFunctionLayerInstance));
-	RegisterAssetTypeActions(MakeShareable(new FAssetTypeActions_MaterialFunctionLayerBlend));
-	RegisterAssetTypeActions(MakeShareable(new FAssetTypeActions_MaterialFunctionLayerBlendInstance));
-	RegisterAssetTypeActions(MakeShareable(new FAssetTypeActions_MaterialFunctionInstance));
-	RegisterAssetTypeActions(MakeShareable(new FAssetTypeActions_MaterialInstanceConstant(EAssetTypeCategories::Type::None)));
-	RegisterAssetTypeActions(MakeShareable(new FAssetTypeActions_MaterialInstanceDynamic));
-	RegisterAssetTypeActions(MakeShareable(new FAssetTypeActions_MaterialInterface));
 	RegisterAssetTypeActions(MakeShareable(new FAssetTypeActions_MirrorDataTable));
 	RegisterAssetTypeActions(MakeShareable(new FAssetTypeActions_ParticleSystem));
 	RegisterAssetTypeActions(MakeShareable(new FAssetTypeActions_PhysicalMaterialMask));
@@ -1583,6 +1567,52 @@ UObject* UAssetToolsImpl::CreateAsset(const FString& AssetName, const FString& P
 	return NewObj;
 }
 
+void UAssetToolsImpl::CreateAssetsFrom(TConstArrayView<UObject*> SourceObjects, UClass* CreateAssetType, const FString& DefaultSuffix, TFunctionRef<UFactory*(UObject*)> FactoryConstructor, FName CallingContext)
+{
+	if ( SourceObjects.Num() == 1 )
+	{
+		if (UObject* SourceObject = SourceObjects[0])
+		{
+			// Create an appropriate and unique name 
+			FString Name;
+			FString PackageName;
+			CreateUniqueAssetName(SourceObject->GetOutermost()->GetName(), DefaultSuffix, PackageName, Name);
+
+			UFactory* Factory = FactoryConstructor(SourceObject);
+
+			FContentBrowserModule& ContentBrowserModule = FModuleManager::LoadModuleChecked<FContentBrowserModule>("ContentBrowser");
+			ContentBrowserModule.Get().CreateNewAsset(Name, FPackageName::GetLongPackagePath(PackageName), CreateAssetType, Factory);
+		}
+	}
+	else
+	{
+		TArray<UObject*> ObjectsToSync;
+		for (UObject* SourceObject : SourceObjects)
+		{
+			if ( SourceObject )
+			{
+				// Determine an appropriate name
+				FString Name;
+				FString PackageName;
+				CreateUniqueAssetName(SourceObject->GetOutermost()->GetName(), DefaultSuffix, PackageName, Name);
+
+				// Create the factory used to generate the asset
+				UFactory* Factory = FactoryConstructor(SourceObject);
+					
+				if (UObject* NewAsset = CreateAsset(Name, FPackageName::GetLongPackagePath(PackageName), CreateAssetType, Factory, CallingContext))
+				{
+					ObjectsToSync.Add(NewAsset);
+				}
+			}
+		}
+
+		if ( ObjectsToSync.Num() > 0 )
+		{
+			SyncBrowserToAssets(ObjectsToSync);
+		}
+	}
+}
+
 UObject* UAssetToolsImpl::CreateAssetWithDialog(UClass* AssetClass, UFactory* Factory, FName CallingContext)
 {
 	if (Factory != nullptr)
diff --git a/Engine/Source/Developer/AssetTools/Private/AssetTools.h b/Engine/Source/Developer/AssetTools/Private/AssetTools.h
--- a/Engine/Source/Developer/AssetTools/Private/AssetTools.h
+++ b/Engine/Source/Developer/AssetTools/Private/AssetTools.h
@@ -81,6 +81,7 @@ public:
 	virtual void GetClassTypeActionsList( TArray<TWeakPtr<IClassTypeActions>>& OutClassTypeActionsList ) const override;
 	virtual TWeakPtr<IClassTypeActions> GetClassTypeActionsForClass( UClass* Class ) const override;
 	virtual UObject* CreateAsset(const FString& AssetName, const FString& PackagePath, UClass* AssetClass, UFactory* Factory, FName CallingContext = NAME_None) override;
+	virtual void CreateAssetsFrom(TConstArrayView<UObject*> SourceObjects, UClass* CreateAssetType, const FString& DefaultSuffix, TFunctionRef<UFactory*(UObject*)> FactoryConstructor, FName CallingContext = NAME_None);
 	virtual UObject* CreateAssetWithDialog(UClass* AssetClass, UFactory* Factory, FName CallingContext = NAME_None) override;
 	virtual UObject* CreateAssetWithDialog(const FString& AssetName, const FString& PackagePath, UClass* AssetClass, UFactory* Factory, FName CallingContext = NAME_None, const bool bCallConfigureProperties = true) override;
 	virtual UObject* DuplicateAsset(const FString& AssetName, const FString& PackagePath, UObject* OriginalObject) override;
diff --git a/Engine/Source/Developer/AssetTools/Private/AssetTypeActions/AssetTypeActions_Material.cpp b/Engine/Source/Developer/AssetTools/Private/AssetTypeActions/AssetTypeActions_Material.cpp
deleted file mode 100644
diff --git a/Engine/Source/Developer/AssetTools/Private/AssetTypeActions/AssetTypeActions_Material.h b/Engine/Source/Developer/AssetTools/Private/AssetTypeActions/AssetTypeActions_Material.h
deleted file mode 100644
diff --git a/Engine/Source/Developer/AssetTools/Private/AssetTypeActions/AssetTypeActions_MaterialFunction.cpp b/Engine/Source/Developer/AssetTools/Private/AssetTypeActions/AssetTypeActions_MaterialFunction.cpp
deleted file mode 100644
diff --git a/Engine/Source/Developer/AssetTools/Private/AssetTypeActions/AssetTypeActions_MaterialFunction.h b/Engine/Source/Developer/AssetTools/Private/AssetTypeActions/AssetTypeActions_MaterialFunction.h
deleted file mode 100644
diff --git a/Engine/Source/Developer/AssetTools/Private/AssetTypeActions/AssetTypeActions_MaterialFunctionInstance.cpp b/Engine/Source/Developer/AssetTools/Private/AssetTypeActions/AssetTypeActions_MaterialFunctionInstance.cpp
deleted file mode 100644
diff --git a/Engine/Source/Developer/AssetTools/Private/AssetTypeActions/AssetTypeActions_MaterialFunctionInstance.h b/Engine/Source/Developer/AssetTools/Private/AssetTypeActions/AssetTypeActions_MaterialFunctionInstance.h
deleted file mode 100644
diff --git a/Engine/Source/Developer/AssetTools/Private/AssetTypeActions/AssetTypeActions_MaterialInstanceConstant.cpp b/Engine/Source/Developer/AssetTools/Private/AssetTypeActions/AssetTypeActions_MaterialInstanceConstant.cpp
deleted file mode 100644
diff --git a/Engine/Source/Developer/AssetTools/Private/AssetTypeActions/AssetTypeActions_MaterialInstanceConstant.h b/Engine/Source/Developer/AssetTools/Private/AssetTypeActions/AssetTypeActions_MaterialInstanceConstant.h
deleted file mode 100644
diff --git a/Engine/Source/Developer/AssetTools/Private/AssetTypeActions/AssetTypeActions_MaterialInstanceDynamic.cpp b/Engine/Source/Developer/AssetTools/Private/AssetTypeActions/AssetTypeActions_MaterialInstanceDynamic.cpp
deleted file mode 100644
diff --git a/Engine/Source/Developer/AssetTools/Private/AssetTypeActions/AssetTypeActions_MaterialInstanceDynamic.h b/Engine/Source/Developer/AssetTools/Private/AssetTypeActions/AssetTypeActions_MaterialInstanceDynamic.h
deleted file mode 100644
diff --git a/Engine/Source/Developer/AssetTools/Private/AssetTypeActions/AssetTypeActions_MaterialInterface.cpp b/Engine/Source/Developer/AssetTools/Private/AssetTypeActions/AssetTypeActions_MaterialInterface.cpp
deleted file mode 100644
diff --git a/Engine/Source/Developer/AssetTools/Private/AssetTypeActions/AssetTypeActions_MaterialInterface.h b/Engine/Source/Developer/AssetTools/Private/AssetTypeActions/AssetTypeActions_MaterialInterface.h
deleted file mode 100644
diff --git a/Engine/Source/Developer/AssetTools/Public/IAssetTools.h b/Engine/Source/Developer/AssetTools/Public/IAssetTools.h
--- a/Engine/Source/Developer/AssetTools/Public/IAssetTools.h
+++ b/Engine/Source/Developer/AssetTools/Public/IAssetTools.h
@@ -219,9 +219,9 @@ struct FMigrationOptions
 		: bPrompt(false)
 		, AssetConflict(EAssetMigrationConflict::Skip)
 	{}
-
 };
 
+
 // An array of maps each storing pairs of original object -> duplicated object
 using FDuplicatedObjects = TArray<TMap<TSoftObjectPtr<UObject>, TSoftObjectPtr<UObject>>>;
 
@@ -294,6 +294,20 @@ public:
 	UFUNCTION(BlueprintCallable, Category = "Editor Scripting | Asset Tools")
 	virtual UObject* CreateAsset(const FString& AssetName, const FString& PackagePath, UClass* AssetClass, UFactory* Factory, FName CallingContext = NAME_None) = 0;
 
+	/**
+	 * Creates one or more assets using the source objects as the basis for the next type.  This is a common enough operation
+	 * it needed a utility.  In the case that you only have a single SourceObject, we'll lean on the content browser
+	 * to create the asset and focus you to it so you can rename it inline.  However in the case that multiple assets
+	 * get created we'll construct each one and then sync the content browser to them.
+	 */
+	virtual void CreateAssetsFrom(TConstArrayView<UObject*> SourceObjects, UClass* CreateAssetType, const FString& DefaultSuffix, TFunctionRef<UFactory*(UObject*)> FactoryConstructor, FName CallingContext = NAME_None) = 0;
+
+	template<typename SourceObjectType, typename = typename TEnableIf<TIsDerivedFrom<SourceObjectType, UObject>::Value>::Type>
+	void CreateAssetsFrom(TConstArrayView<SourceObjectType*> SourceObjects, UClass* CreateAssetType, const FString& DefaultSuffix, TFunctionRef<UFactory*(SourceObjectType*)> FactoryConstructor, FName CallingContext = NAME_None)
+	{
+		CreateAssetsFrom(TConstArrayView<UObject*>(reinterpret_cast<UObject* const*>(SourceObjects.GetData()), SourceObjects.Num()), CreateAssetType, DefaultSuffix, [FactoryConstructor](UObject* SourceObject){ return FactoryConstructor(Cast<SourceObjectType>(SourceObject)); }, CallingContext);
+	}
+
 	/** Opens an asset picker dialog and creates an asset with the specified name and path */
 	UFUNCTION(BlueprintCallable, Category = "Editor Scripting | Asset Tools")
 	virtual UObject* CreateAssetWithDialog(const FString& AssetName, const FString& PackagePath, UClass* AssetClass, UFactory* Factory, FName CallingContext = NAME_None, const bool bCallConfigureProperties = true) = 0;
@@ -562,6 +576,7 @@ public:
 	/** Syncs the primary content browser to the specified assets, whether or not it is locked. Most syncs that come from AssetTools -feel- like they came from the content browser, so this is okay. */
 	virtual void SyncBrowserToAssets(const TArray<UObject*>& AssetsToSync) = 0;
 	virtual void SyncBrowserToAssets(const TArray<FAssetData>& AssetsToSync) = 0;
+	
 };
 
 UCLASS(transient)
diff --git a/Engine/Source/Editor/Blutility/Classes/ActorActionUtility.h b/Engine/Source/Editor/Blutility/Classes/ActorActionUtility.h
--- a/Engine/Source/Editor/Blutility/Classes/ActorActionUtility.h
+++ b/Engine/Source/Editor/Blutility/Classes/ActorActionUtility.h
@@ -37,8 +37,7 @@ public:
 	
 private:
 	/**
-	 * For simple Asset Action's you should fill out the supported class here.  Don't bother with GetSupportedClass()
-	 * * unless you actually need to do specialized dynamic logic.
+	 * For simple Asset Action's you should fill out the supported class here.
 	 */
 	UPROPERTY(EditDefaultsOnly, Category="Assets", meta=(AllowAbstract))
 	TArray<TSoftClassPtr<UObject>> SupportedClasses;
diff --git a/Engine/Source/Editor/Blutility/Classes/AssetActionUtility.h b/Engine/Source/Editor/Blutility/Classes/AssetActionUtility.h
--- a/Engine/Source/Editor/Blutility/Classes/AssetActionUtility.h
+++ b/Engine/Source/Editor/Blutility/Classes/AssetActionUtility.h
@@ -14,10 +14,34 @@ namespace AssetActionUtilityTags
 	extern const FName SupportedClasses;
 	extern const FName IsActionForBlueprint;
 	extern const FName CallableFunctions;
+	extern const FName SupportedConditions;
 }
 
 struct FAssetData;
 
+USTRUCT()
+struct FAssetActionSupportCondition
+{
+	GENERATED_BODY()
+
+	/**
+	 * This is a content browser styled filter.  For example, ..._N AND VirtualTextureStreaming=FALSE, would require that
+	 * asset tag VirtualTextureStreaming be false, and the asset name end in _N.
+	 *
+	 * You can learn more about the content browser search syntax in the "Advanced Search Syntax" documentation.
+	 */
+	UPROPERTY(EditAnywhere, Category=Condition)
+	FString Filter;
+
+	/**
+	 * This is the failure reason to reply to the user with if the condition above fails.
+	 * If you leave this blank, we will hide the option.  If fill in the reason, will show the option but give the reason
+	 * it is disabled.
+	 */
+	UPROPERTY(EditAnywhere, Category=Condition, meta=(MultiLine=true))
+	FString FailureReason;
+};
+
 /** 
  * Base class for all asset action-related utilities
  * Any functions/events that are exposed on derived classes that have the correct signature will be
@@ -58,9 +82,17 @@ public:
 
 protected:
 	/**
-	 * For simple Asset Action's you should fill out the supported class here.  Don't bother with GetSupportedClass()
-	 * * unless you actually need to do specialized dynamic logic.
+	 * The supported classes controls the list of classes that may be operated on by all of the asset functions in this
+	 * utility class.
 	 */
-	UPROPERTY(EditDefaultsOnly, Category="Assets", meta=(AllowAbstract))
+	UPROPERTY(EditDefaultsOnly, Category="Asset Support", meta=(AllowAbstract))
 	TArray<TSoftClassPtr<UObject>> SupportedClasses;
+
+	/**
+	 * The supported conditions for any asset to use these utility functions.  Note: all of these conditions must pass
+	 * in sequence.  So if you have earlier failure conditions you want to be run first, put them first in the list,
+	 * if those fail, their failure reason will be handled first.
+	 */
+	UPROPERTY(EditDefaultsOnly, Category="Asset Support")
+	TArray<FAssetActionSupportCondition> SupportedConditions;
 };
\ No newline at end of file
diff --git a/Engine/Source/Editor/Blutility/Private/AssetActionUtility.cpp b/Engine/Source/Editor/Blutility/Private/AssetActionUtility.cpp
--- a/Engine/Source/Editor/Blutility/Private/AssetActionUtility.cpp
+++ b/Engine/Source/Editor/Blutility/Private/AssetActionUtility.cpp
@@ -24,6 +24,7 @@ void UAssetActionUtility::GetAssetRegistryTags(TArray<FAssetRegistryTag>& OutTag
 	}
 	
 	FAssetActionUtilityPrototype::AddTagsFor_SupportedClasses(SupportedClasses, OutTags);
+	FAssetActionUtilityPrototype::AddTagsFor_SupportedConditions(SupportedConditions, OutTags);
 	FAssetActionUtilityPrototype::AddTagsFor_IsActionForBlueprints(IsActionForBlueprints(), OutTags);
 	FAssetActionUtilityPrototype::AddTagsFor_CallableFunctions(this, OutTags);
 }
diff --git a/Engine/Source/Editor/Blutility/Private/BlutilityMenuExtensions.cpp b/Engine/Source/Editor/Blutility/Private/BlutilityMenuExtensions.cpp
--- a/Engine/Source/Editor/Blutility/Private/BlutilityMenuExtensions.cpp
+++ b/Engine/Source/Editor/Blutility/Private/BlutilityMenuExtensions.cpp
@@ -17,6 +17,7 @@
 #include "EditorUtilityAssetPrototype.h"
 #include "Editor/EditorEngine.h"
 #include "EditorUtilityBlueprint.h"
+#include "FrontendFilters.h"
 #include "Engine/Blueprint.h"
 #include "Framework/Application/SlateApplication.h"
 #include "Framework/Commands/UIAction.h"
@@ -72,6 +73,340 @@ class IToolkitHost;
 
 #define LOCTEXT_NAMESPACE "BlutilityMenuExtensions"
 
+namespace BlutilityUtil
+{
+	/** Mapping of asset property tag aliases that can be used by text searches */
+	class FAssetPropertyTagAliases
+	{
+	public:
+		static FAssetPropertyTagAliases& Get()
+		{
+			static FAssetPropertyTagAliases Singleton;
+			return Singleton;
+		}
+
+		/** Get the source tag for the given asset data and alias, or none if there is no match */
+		FName GetSourceTagFromAlias(const FAssetData& InAssetData, const FName InAlias)
+		{
+			TSharedPtr<TMap<FName, FName>>& AliasToSourceTagMapping = ClassToAliasTagsMapping.FindOrAdd(InAssetData.AssetClassPath);
+
+			if (!AliasToSourceTagMapping.IsValid())
+			{
+				static const FName NAME_DisplayName(TEXT("DisplayName"));
+
+				AliasToSourceTagMapping = MakeShared<TMap<FName, FName>>();
+
+				UClass* AssetClass = InAssetData.GetClass();
+				if (AssetClass)
+				{
+					TMap<FName, UObject::FAssetRegistryTagMetadata> AssetTagMetaData;
+					AssetClass->GetDefaultObject()->GetAssetRegistryTagMetadata(AssetTagMetaData);
+
+					for (const auto& AssetTagMetaDataPair : AssetTagMetaData)
+					{
+						if (!AssetTagMetaDataPair.Value.DisplayName.IsEmpty())
+						{
+							const FName DisplayName = MakeObjectNameFromDisplayLabel(AssetTagMetaDataPair.Value.DisplayName.ToString(), NAME_None);
+							AliasToSourceTagMapping->Add(DisplayName, AssetTagMetaDataPair.Key);
+						}
+					}
+
+					for (const auto& KeyValuePair : InAssetData.TagsAndValues)
+					{
+						if (FProperty* Field = FindFProperty<FProperty>(AssetClass, KeyValuePair.Key))
+						{
+							if (Field->HasMetaData(NAME_DisplayName))
+							{
+								const FName DisplayName = MakeObjectNameFromDisplayLabel(Field->GetMetaData(NAME_DisplayName), NAME_None);
+								AliasToSourceTagMapping->Add(DisplayName, KeyValuePair.Key);
+							}
+						}
+					}
+				}
+			}
+
+			return AliasToSourceTagMapping.IsValid() ? AliasToSourceTagMapping->FindRef(InAlias) : NAME_None;
+		}
+
+	private:
+		/** Mapping from class name -> (alias -> source) */
+		TMap<FTopLevelAssetPath, TSharedPtr<TMap<FName, FName>>> ClassToAliasTagsMapping;
+	};
+	
+	/** Expression context to test the given asset data against the current text filter */
+	class FTextFilterExpressionContext : public ITextFilterExpressionContext
+	{
+	public:
+		typedef TRemoveReference<const FAssetData&>::Type* FAssetFilterTypePtr;
+
+		FTextFilterExpressionContext()
+			: AssetPtr(nullptr)
+			, bIncludeClassName(true)
+			, bIncludeAssetPath(false)
+			, bIncludeCollectionNames(true)
+			, NameKeyName("Name")
+			, PathKeyName("Path")
+			, ClassKeyName("Class")
+			, TypeKeyName("Type")
+			, CollectionKeyName("Collection")
+			, TagKeyName("Tag")
+		{
+		}
+
+		void SetAsset(FAssetFilterTypePtr InAsset) 
+		{
+			AssetPtr = InAsset;
+
+			if (bIncludeAssetPath)
+			{
+				// Get the full asset path, and also split it so we can compare each part in the filter
+				AssetPtr->PackageName.AppendString(AssetFullPath);
+				AssetFullPath.ParseIntoArray(AssetSplitPath, TEXT("/"));
+				AssetFullPath.ToUpperInline();
+
+				if (bIncludeClassName)
+				{
+					// Get the full export text path as people sometimes search by copying this (requires class and asset path search to be enabled in order to match)
+					AssetPtr->GetExportTextName(AssetExportTextName);
+					AssetExportTextName.ToUpperInline();
+				}
+			}
+		}
+
+		const FAssetData* GetAsset() const { return AssetPtr; }
+
+		void ClearAsset()
+		{
+			AssetPtr = nullptr;
+			AssetFullPath.Reset();
+			AssetExportTextName.Reset();
+			AssetSplitPath.Reset();
+			AssetCollectionNames.Reset();
+		}
+
+		void SetIncludeClassName(const bool InIncludeClassName)
+		{
+			bIncludeClassName = InIncludeClassName;
+		}
+
+		bool GetIncludeClassName() const
+		{
+			return bIncludeClassName;
+		}
+
+		void SetIncludeAssetPath(const bool InIncludeAssetPath)
+		{
+			bIncludeAssetPath = InIncludeAssetPath;
+		}
+
+		bool GetIncludeAssetPath() const
+		{
+			return bIncludeAssetPath;
+		}
+
+		void SetIncludeCollectionNames(const bool InIncludeCollectionNames)
+		{
+			bIncludeCollectionNames = InIncludeCollectionNames;
+		}
+
+		bool GetIncludeCollectionNames() const
+		{
+			return bIncludeCollectionNames;
+		}
+
+		virtual bool TestBasicStringExpression(const FTextFilterString& InValue, const ETextFilterTextComparisonMode InTextComparisonMode) const override
+		{
+			if (InValue.CompareName(AssetPtr->AssetName, InTextComparisonMode))
+			{
+				return true;
+			}
+
+			if (bIncludeAssetPath)
+			{
+				if (InValue.CompareFString(AssetFullPath, InTextComparisonMode))
+				{
+					return true;
+				}
+
+				for (const FString& AssetPathPart : AssetSplitPath)
+				{
+					if (InValue.CompareFString(AssetPathPart, InTextComparisonMode))
+					{
+						return true;
+					}
+				}
+			}
+
+			if (bIncludeClassName)
+			{
+				if (InValue.CompareFString(AssetPtr->AssetClassPath.ToString(), InTextComparisonMode))
+				{
+					return true;
+				}
+			}
+
+			if (bIncludeClassName && bIncludeAssetPath)
+			{
+				// Only test this if we're searching the class name and asset path too, as the exported text contains the type and path in the string
+				if (InValue.CompareFString(AssetExportTextName, InTextComparisonMode))
+				{
+					return true;
+				}
+			}
+
+			if (bIncludeCollectionNames)
+			{
+				for (const FName& AssetCollectionName : AssetCollectionNames)
+				{
+					if (InValue.CompareName(AssetCollectionName, InTextComparisonMode))
+					{
+						return true;
+					}
+				}
+			}
+
+			return false;
+		}
+
+		virtual bool TestComplexExpression(const FName& InKey, const FTextFilterString& InValue, const ETextFilterComparisonOperation InComparisonOperation, const ETextFilterTextComparisonMode InTextComparisonMode) const override
+		{
+			// Special case for the asset name, as this isn't contained within the asset registry meta-data
+			if (InKey == NameKeyName)
+			{
+				// Names can only work with Equal or NotEqual type tests
+				if (InComparisonOperation != ETextFilterComparisonOperation::Equal && InComparisonOperation != ETextFilterComparisonOperation::NotEqual)
+				{
+					return false;
+				}
+
+				const bool bIsMatch = TextFilterUtils::TestBasicStringExpression(AssetPtr->AssetName, InValue, InTextComparisonMode);
+				return (InComparisonOperation == ETextFilterComparisonOperation::Equal) ? bIsMatch : !bIsMatch;
+			}
+
+			// Special case for the asset path, as this isn't contained within the asset registry meta-data
+			if (InKey == PathKeyName)
+			{
+				// Paths can only work with Equal or NotEqual type tests
+				if (InComparisonOperation != ETextFilterComparisonOperation::Equal && InComparisonOperation != ETextFilterComparisonOperation::NotEqual)
+				{
+					return false;
+				}
+
+				// If the comparison mode is partial, then we only need to test the ObjectPath as that contains the other two as sub-strings
+				bool bIsMatch = false;
+				if (InTextComparisonMode == ETextFilterTextComparisonMode::Partial)
+				{
+					bIsMatch = TextFilterUtils::TestBasicStringExpression(AssetPtr->GetObjectPathString(), InValue, InTextComparisonMode);
+				}
+				else
+				{
+					bIsMatch = TextFilterUtils::TestBasicStringExpression(AssetPtr->GetObjectPathString(), InValue, InTextComparisonMode)
+						|| TextFilterUtils::TestBasicStringExpression(AssetPtr->PackageName, InValue, InTextComparisonMode)
+						|| TextFilterUtils::TestBasicStringExpression(AssetPtr->PackagePath, InValue, InTextComparisonMode);
+				}
+				return (InComparisonOperation == ETextFilterComparisonOperation::Equal) ? bIsMatch : !bIsMatch;
+			}
+
+			// Special case for the asset type, as this isn't contained within the asset registry meta-data
+			if (InKey == ClassKeyName || InKey == TypeKeyName)
+			{
+				// Class names can only work with Equal or NotEqual type tests
+				if (InComparisonOperation != ETextFilterComparisonOperation::Equal && InComparisonOperation != ETextFilterComparisonOperation::NotEqual)
+				{
+					return false;
+				}
+
+				const bool bIsMatch = TextFilterUtils::TestBasicStringExpression(AssetPtr->AssetClassPath.ToString(), InValue, InTextComparisonMode);
+				return (InComparisonOperation == ETextFilterComparisonOperation::Equal) ? bIsMatch : !bIsMatch;
+			}
+
+			// Special case for collections, as these aren't contained within the asset registry meta-data
+			if (InKey == CollectionKeyName || InKey == TagKeyName)
+			{
+				// Collections can only work with Equal or NotEqual type tests
+				if (InComparisonOperation != ETextFilterComparisonOperation::Equal && InComparisonOperation != ETextFilterComparisonOperation::NotEqual)
+				{
+					return false;
+				}
+
+				bool bFoundMatch = false;
+				for (const FName& AssetCollectionName : AssetCollectionNames)
+				{
+					if (TextFilterUtils::TestBasicStringExpression(AssetCollectionName, InValue, InTextComparisonMode))
+					{
+						bFoundMatch = true;
+						break;
+					}
+				}
+
+				return (InComparisonOperation == ETextFilterComparisonOperation::Equal) ? bFoundMatch : !bFoundMatch;
+			}
+
+			// Generic handling for anything in the asset meta-data
+			{
+				auto GetMetaDataValue = [this, &InKey](FString& OutMetaDataValue) -> bool
+				{
+					// Check for a literal key
+					if (AssetPtr->GetTagValue(InKey, OutMetaDataValue))
+					{
+						return true;
+					}
+
+					// Check for an alias key
+					const FName LiteralKey = FAssetPropertyTagAliases::Get().GetSourceTagFromAlias(*AssetPtr, InKey);
+					if (!LiteralKey.IsNone() && AssetPtr->GetTagValue(LiteralKey, OutMetaDataValue))
+					{
+						return true;
+					}
+
+					return false;
+				};
+
+				FString MetaDataValue;
+				if (GetMetaDataValue(MetaDataValue))
+				{
+					return TextFilterUtils::TestComplexExpression(MetaDataValue, InValue, InComparisonOperation, InTextComparisonMode);
+				}
+			}
+
+			return false;
+		}
+
+	private:
+		/** Pointer to the asset we're currently filtering */
+		FAssetFilterTypePtr AssetPtr;
+
+		/** Full path of the current asset */
+		FString AssetFullPath;
+
+		/** The export text name of the current asset */
+		FString AssetExportTextName;
+
+		/** Split path of the current asset */
+		TArray<FString> AssetSplitPath;
+
+		/** Names of the collections that the current asset is in */
+		TArray<FName> AssetCollectionNames;
+
+		/** Are we supposed to include the class name in our basic string tests? */
+		bool bIncludeClassName;
+
+		/** Search inside the entire asset path? */
+		bool bIncludeAssetPath;
+
+		/** Search collection names? */
+		bool bIncludeCollectionNames;
+
+		/** Keys used by TestComplexExpression */
+		const FName NameKeyName;
+		const FName PathKeyName;
+		const FName ClassKeyName;
+		const FName TypeKeyName;
+		const FName CollectionKeyName;
+		const FName TagKeyName;
+	};	
+}
+
 /** Dialog widget used to display function properties */
 class SFunctionParamDialog : public SCompoundWidget
 {
@@ -329,9 +664,6 @@ void FBlutilityMenuExtensions::OpenEditorForUtility(const FFunctionAndUtil& Func
 
 void FBlutilityMenuExtensions::ExtractFunctions(TMap<TSharedRef<FAssetActionUtilityPrototype>, TSet<int32>>& Utils, TMap<FString, TArray<FFunctionAndUtil>>& OutCategoryFunctions)
 {
-	TSet<UClass*> ProcessedClasses;
-	const static FName NAME_CallInEditor(TEXT("CallInEditor"));
-
 	// Find the exposed functions available in each class, making sure to not list shared functions from a parent class more than once
 	for (TPair<TSharedRef<FAssetActionUtilityPrototype>, TSet<int32>> UtilitySelectionPair : Utils)
 	{
@@ -357,113 +689,164 @@ void FBlutilityMenuExtensions::CreateBlutilityActionsMenu(FMenuBuilder& MenuBuil
 {
 	TMap<FString, TArray<FFunctionAndUtil>> CategoryFunctions;
 	ExtractFunctions(Utils, CategoryFunctions);
-
+	
 	auto AddFunctionEntries = [Selection, IsValidPropertyType](FMenuBuilder& SubMenuBuilder, const TArray<FFunctionAndUtil>& FunctionUtils)
 	{
+		BlutilityUtil::FTextFilterExpressionContext TextFilterContext;
+		
+		FTextFilterExpressionEvaluator TextFilterExpressionEvaluator(ETextFilterExpressionEvaluatorMode::Complex);
+		
 		for (const FFunctionAndUtil& FunctionAndUtil : FunctionUtils)
 		{
-			const FText TooltipText = FText::Format(LOCTEXT("AssetUtilTooltipFormat", "{0}\n(Shift-click to edit script)"), FunctionAndUtil.FunctionData.TooltipText);
+			bool PassesFilterCondition = true;
+			FString FilterFailureMessage;
+			
+			if constexpr ( std::is_same_v<FAssetData, SelectionType> )
+			{
+				TArray<FAssetActionSupportCondition> Conditions = FunctionAndUtil.Util->GetAssetActionSupportConditions();
 
-			SubMenuBuilder.AddMenuEntry(
-				FunctionAndUtil.FunctionData.NameText,
-				TooltipText,
-				FSlateIcon(FAppStyle::GetAppStyleSetName(), "GraphEditor.Event_16x"),
-				FExecuteAction::CreateLambda([FunctionAndUtil, Selection, IsValidPropertyType]
+				for (const FAssetActionSupportCondition& Condition : Conditions)
 				{
-					if (FSlateApplication::Get().GetModifierKeys().IsShiftDown())
+					TextFilterExpressionEvaluator.SetFilterText(FText::FromString(Condition.Filter));
+				
+					for (const int32& SelectionIndex : FunctionAndUtil.SelectionIndices)
 					{
-						OpenEditorForUtility(FunctionAndUtil);
+						const auto SelectedAsset = Selection[SelectionIndex];
+						TextFilterContext.SetAsset(&SelectedAsset);
+						if (!TextFilterExpressionEvaluator.TestTextFilter(TextFilterContext))
+						{
+							PassesFilterCondition = false;
+							FilterFailureMessage = Condition.FailureReason;
+							break;
+						}
 					}
-					else
+
+					if (!PassesFilterCondition)
 					{
-						// We dont run this on the CDO, as bad things could occur!
-						UObject* TempObject = NewObject<UObject>(GetTransientPackage(), Cast<UObject>(FunctionAndUtil.Util->LoadUtilityAsset())->GetClass());
-						TempObject->AddToRoot(); // Some Blutility actions might run GC so the TempObject needs to be rooted to avoid getting destroyed
+						break;
+					}
+				}
+			}
+			
+			FText TooltipText;
 
-						UFunction* Function = FunctionAndUtil.GetFunction();
-						
-						if (Function->NumParms > 0)
-						{
-							// Create a parameter struct and fill in defaults
-							TSharedRef<FStructOnScope> FuncParams = MakeShared<FStructOnScope>(Function);
+			if (FilterFailureMessage.IsEmpty())
+			{
+				TooltipText = FText::Format(LOCTEXT("AssetUtilTooltipFormat", "{0}\n\n(Shift-click to edit script)"), FunctionAndUtil.FunctionData.TooltipText);	
+			}
+			else
+			{
+				TooltipText = FText::Format(LOCTEXT("AssetUtilTooltipWithErrorFormat", "{0}\n\n({1})\n\n(Shift-click to edit script)"), FunctionAndUtil.FunctionData.TooltipText, FText::FromString(FilterFailureMessage));
+			}
 
-							FProperty* FirstParamProperty = nullptr;
+			SubMenuBuilder.AddMenuEntry(
+				FunctionAndUtil.FunctionData.NameText,
+				TooltipText,
+				FSlateIcon(FAppStyle::GetAppStyleSetName(), "GraphEditor.Event_16x"),
+				
+				FUIAction
+				(
+					FExecuteAction::CreateLambda([FunctionAndUtil, Selection, IsValidPropertyType]
+					{
+						if (FSlateApplication::Get().GetModifierKeys().IsShiftDown())
+						{
+							OpenEditorForUtility(FunctionAndUtil);
+						}
+						else
+						{
+							// We dont run this on the CDO, as bad things could occur!
+							UObject* TempObject = NewObject<UObject>(GetTransientPackage(), Cast<UObject>(FunctionAndUtil.Util->LoadUtilityAsset())->GetClass());
+							TempObject->AddToRoot(); // Some Blutility actions might run GC so the TempObject needs to be rooted to avoid getting destroyed
 
-							int32 ParameterIndex = 0;
-							for (TFieldIterator<FProperty> It(Function); It&& It->HasAnyPropertyFlags(CPF_Parm); ++It)
+							UFunction* Function = FunctionAndUtil.GetFunction();
+							
+							if (Function->NumParms > 0)
 							{
-								FString Defaults;
-								if (UEdGraphSchema_K2::FindFunctionParameterDefaultValue(Function, *It, Defaults))
-								{
-									It->ImportText_Direct(*Defaults, It->ContainerPtrToValuePtr<uint8>(FuncParams->GetStructMemory()), nullptr, PPF_None);
-								}
+								// Create a parameter struct and fill in defaults
+								TSharedRef<FStructOnScope> FuncParams = MakeShared<FStructOnScope>(Function);
 
-								// Check to see if the first parameter matches the selection object type, in that case we can directly forward the selection to it
-								if (ParameterIndex == 0 && IsValidPropertyType(*It))
-								{
-									FirstParamProperty = *It;
-								}
+								FProperty* FirstParamProperty = nullptr;
 
-								++ParameterIndex;
-							}
+								int32 ParameterIndex = 0;
+								for (TFieldIterator<FProperty> It(Function); It&& It->HasAnyPropertyFlags(CPF_Parm); ++It)
+								{
+									FString Defaults;
+									if (UEdGraphSchema_K2::FindFunctionParameterDefaultValue(Function, *It, Defaults))
+									{
+										It->ImportText_Direct(*Defaults, It->ContainerPtrToValuePtr<uint8>(FuncParams->GetStructMemory()), nullptr, PPF_None);
+									}
 
-							bool bApply = true;
+									// Check to see if the first parameter matches the selection object type, in that case we can directly forward the selection to it
+									if (ParameterIndex == 0 && IsValidPropertyType(*It))
+									{
+										FirstParamProperty = *It;
+									}
 
-							if (!FirstParamProperty || ParameterIndex > 1)
-							{
-								// pop up a dialog to input params to the function
-								TSharedRef<SWindow> Window = SNew(SWindow)
-									.Title(Function->GetDisplayNameText())
-									.ClientSize(FVector2D(400, 200))
-									.SupportsMinimize(false)
-									.SupportsMaximize(false);
-
-								TSharedPtr<SFunctionParamDialog> Dialog;
-								Window->SetContent(
-									SAssignNew(Dialog, SFunctionParamDialog, Window, FuncParams, FirstParamProperty ? FirstParamProperty->GetFName() : NAME_None)
-									.OkButtonText(LOCTEXT("OKButton", "OK"))
-									.OkButtonTooltipText(Function->GetToolTipText()));
-
-								GEditor->EditorAddModalWindow(Window);
-								bApply = Dialog->bOKPressed;
-							}
+									++ParameterIndex;
+								}
 
+								bool bApply = true;
 
-							if (bApply)
-							{
-								FScopedTransaction Transaction(NSLOCTEXT("UnrealEd", "BlutilityAction", "Blutility Action"));
-								FEditorScriptExecutionGuard ScriptGuard;
-								const bool bForwardUserSelection = FirstParamProperty != nullptr;
-								if (bForwardUserSelection)
+								if (!FirstParamProperty || ParameterIndex > 1)
 								{
-									// For each user-select asset forward the selection object into the function first's parameter (if it matches)
-									const FString Path = FirstParamProperty->GetPathName(Function);
+									// pop up a dialog to input params to the function
+									TSharedRef<SWindow> Window = SNew(SWindow)
+										.Title(Function->GetDisplayNameText())
+										.ClientSize(FVector2D(400, 200))
+										.SupportsMinimize(false)
+										.SupportsMaximize(false);
+
+									TSharedPtr<SFunctionParamDialog> Dialog;
+									Window->SetContent(
+										SAssignNew(Dialog, SFunctionParamDialog, Window, FuncParams, FirstParamProperty ? FirstParamProperty->GetFName() : NAME_None)
+										.OkButtonText(LOCTEXT("OKButton", "OK"))
+										.OkButtonTooltipText(Function->GetToolTipText()));
+
+									GEditor->EditorAddModalWindow(Window);
+									bApply = Dialog->bOKPressed;
+								}
 
-									// Ensure we only process selection objects that are valid for this function/utility
-									for (const int32& SelectionIndex : FunctionAndUtil.SelectionIndices)
+
+								if (bApply)
+								{
+									FScopedTransaction Transaction(NSLOCTEXT("UnrealEd", "BlutilityAction", "Blutility Action"));
+									FEditorScriptExecutionGuard ScriptGuard;
+									const bool bForwardUserSelection = FirstParamProperty != nullptr;
+									if (bForwardUserSelection)
+									{
+										// For each user-select asset forward the selection object into the function first's parameter (if it matches)
+										const FString Path = FirstParamProperty->GetPathName(Function);
+
+										// Ensure we only process selection objects that are valid for this function/utility
+										for (const int32& SelectionIndex : FunctionAndUtil.SelectionIndices)
+										{
+											const auto SelectedAsset = Selection[SelectionIndex];
+											FirstParamProperty->CopySingleValue(FirstParamProperty->ContainerPtrToValuePtr<uint8>(FuncParams->GetStructMemory()), &SelectedAsset);
+											TempObject->ProcessEvent(Function, FuncParams->GetStructMemory());
+										}
+									}
+									else
 									{
-										const auto SelectedAsset = Selection[SelectionIndex];
-										FirstParamProperty->CopySingleValue(FirstParamProperty->ContainerPtrToValuePtr<uint8>(FuncParams->GetStructMemory()), &SelectedAsset);
+										// User is expected to manage the asset selection on its own
 										TempObject->ProcessEvent(Function, FuncParams->GetStructMemory());
 									}
 								}
-								else
-								{
-									// User is expected to manage the asset selection on its own
-									TempObject->ProcessEvent(Function, FuncParams->GetStructMemory());
-								}
 							}
-						}
-						else
-						{
-							FScopedTransaction Transaction(NSLOCTEXT("UnrealEd", "BlutilityAction", "Blutility Action"));
-							FEditorScriptExecutionGuard ScriptGuard;
-							TempObject->ProcessEvent(Function, nullptr);
-						}
+							else
+							{
+								FScopedTransaction Transaction(NSLOCTEXT("UnrealEd", "BlutilityAction", "Blutility Action"));
+								FEditorScriptExecutionGuard ScriptGuard;
+								TempObject->ProcessEvent(Function, nullptr);
+							}
 
-						TempObject->RemoveFromRoot();
-					}
-				})
+							TempObject->RemoveFromRoot();
+						}
+					}),
+					FCanExecuteAction::CreateLambda([PassesFilterCondition]()
+					{
+						return PassesFilterCondition;
+					})
+				)
 			);
 		}
 	};
diff --git a/Engine/Source/Editor/Blutility/Private/EditorUtilityAssetPrototype.cpp b/Engine/Source/Editor/Blutility/Private/EditorUtilityAssetPrototype.cpp
--- a/Engine/Source/Editor/Blutility/Private/EditorUtilityAssetPrototype.cpp
+++ b/Engine/Source/Editor/Blutility/Private/EditorUtilityAssetPrototype.cpp
@@ -1,6 +1,8 @@
 // Copyright Epic Games, Inc. All Rights Reserved.
 
 #include "EditorUtilityAssetPrototype.h"
+
+#include "AssetActionUtility.h"
 #include "EditorUtilityBlueprint.h"
 #include "JsonObjectConverter.h"
 #include "AssetRegistry/AssetData.h"
@@ -10,6 +12,7 @@ const FName AssetActionUtilityTags::BlutilityTagVersion(TEXT("BlutilityTagVersio
 const FName AssetActionUtilityTags::SupportedClasses(TEXT("SupportedClasses"));
 const FName AssetActionUtilityTags::IsActionForBlueprint(TEXT("IsActionForBlueprint"));
 const FName AssetActionUtilityTags::CallableFunctions(TEXT("CallableFunctions"));
+const FName AssetActionUtilityTags::SupportedConditions(TEXT("SupportedConditions"));
 
 const int32 AssetActionUtilityTags::TagVersion = 1;
 
@@ -83,6 +86,19 @@ TArray<FBlutilityFunctionData> FAssetActionUtilityPrototype::GetCallableFunction
 	return FunctionDatas;
 }
 
+TArray<FAssetActionSupportCondition> FAssetActionUtilityPrototype::GetAssetActionSupportConditions() const
+{
+	TArray<FAssetActionSupportCondition> Results;
+	
+	FString Json;
+	if (UtilityBlueprintAsset.GetTagValue(AssetActionUtilityTags::SupportedConditions, Json))
+	{
+		FJsonObjectConverter::JsonArrayStringToUStruct(Json, &Results, 0, 0);
+	}
+
+	return Results;
+}
+
 void FAssetActionUtilityPrototype::AddTagsFor_Version(TArray<UObject::FAssetRegistryTag>& OutTags)
 {
 	// Adding a version to the tags just in case we need to know 'this blutility is out of date and we can't go based off the cached data because the format isn't followed, or this is a pre-tagged version'
@@ -95,6 +111,31 @@ void FAssetActionUtilityPrototype::AddTagsFor_SupportedClasses(const TArray<TSof
 	OutTags.Add(UObject::FAssetRegistryTag(AssetActionUtilityTags::SupportedClasses, SupportedClassesString, UObject::FAssetRegistryTag::TT_Hidden));
 }
 
+void FAssetActionUtilityPrototype::AddTagsFor_SupportedConditions(const TArray<FAssetActionSupportCondition>& SupportedConditions, TArray<UObject::FAssetRegistryTag>& OutTags)
+{
+	TArray<TSharedPtr<FJsonValue>> SupportedConditionJsonValues;
+	for (const FAssetActionSupportCondition& Condition : SupportedConditions)
+	{
+		TSharedPtr<FJsonObject> ConditionJO = FJsonObjectConverter::UStructToJsonObject(Condition);
+		if (ConditionJO.IsValid())
+		{
+			SupportedConditionJsonValues.Add(MakeShared<FJsonValueObject>(ConditionJO));
+		}
+	}
+
+	if (SupportedConditionJsonValues.Num() > 0)
+	{
+		FString SupportedConditionsJson;
+		{
+			TSharedRef<FJsonValueArray> JsonArray = MakeShared<FJsonValueArray>(SupportedConditionJsonValues);
+			TSharedRef<TJsonWriter<TCHAR, TCondensedJsonPrintPolicy<TCHAR>>> JsonWriter = TJsonWriterFactory<TCHAR, TCondensedJsonPrintPolicy<TCHAR>>::Create(&SupportedConditionsJson);
+			FJsonSerializer::Serialize(JsonArray, FString(), JsonWriter);
+		}
+
+		OutTags.Add(UObject::FAssetRegistryTag(AssetActionUtilityTags::SupportedConditions, SupportedConditionsJson, UObject::FAssetRegistryTag::TT_Hidden));
+	}
+}
+
 void FAssetActionUtilityPrototype::AddTagsFor_IsActionForBlueprints(bool IsActionForBlueprints, TArray<UObject::FAssetRegistryTag>& OutTags)
 {
 	const bool ActionForBlueprint = IsActionForBlueprints;
@@ -128,9 +169,9 @@ void FAssetActionUtilityPrototype::AddTagsFor_CallableFunctions(const UObject* F
 	}
 
 	TArray<TSharedPtr<FJsonValue>> FunctionDataJsonValues;
-	for (const FBlutilityFunctionData& Highlight : FunctionDatas)
+	for (const FBlutilityFunctionData& FunctionData : FunctionDatas)
 	{
-		TSharedPtr<FJsonObject> FunctionDataJO = FJsonObjectConverter::UStructToJsonObject(Highlight);
+		TSharedPtr<FJsonObject> FunctionDataJO = FJsonObjectConverter::UStructToJsonObject(FunctionData);
 		if (FunctionDataJO.IsValid())
 		{
 			FunctionDataJsonValues.Add(MakeShared<FJsonValueObject>(FunctionDataJO));
@@ -141,9 +182,9 @@ void FAssetActionUtilityPrototype::AddTagsFor_CallableFunctions(const UObject* F
 	{
 		FString FunctionDataJson;
 		{
-			TSharedRef<FJsonValueArray> HighlightValuesArrayValue = MakeShared<FJsonValueArray>(FunctionDataJsonValues);
+			TSharedRef<FJsonValueArray> JsonArray = MakeShared<FJsonValueArray>(FunctionDataJsonValues);
 			TSharedRef<TJsonWriter<TCHAR, TCondensedJsonPrintPolicy<TCHAR>>> JsonWriter = TJsonWriterFactory<TCHAR, TCondensedJsonPrintPolicy<TCHAR>>::Create(&FunctionDataJson);
-			FJsonSerializer::Serialize(HighlightValuesArrayValue, FString(), JsonWriter);
+			FJsonSerializer::Serialize(JsonArray, FString(), JsonWriter);
 		}
 
 		OutTags.Add(UObject::FAssetRegistryTag(AssetActionUtilityTags::CallableFunctions, FunctionDataJson, UObject::FAssetRegistryTag::TT_Hidden));
diff --git a/Engine/Source/Editor/Blutility/Private/EditorUtilityAssetPrototype.h b/Engine/Source/Editor/Blutility/Private/EditorUtilityAssetPrototype.h
--- a/Engine/Source/Editor/Blutility/Private/EditorUtilityAssetPrototype.h
+++ b/Engine/Source/Editor/Blutility/Private/EditorUtilityAssetPrototype.h
@@ -2,6 +2,7 @@
 
 #pragma once
 
+#include "AssetActionUtility.h"
 #include "UObject/Object.h"
 #include "AssetRegistry/AssetData.h"
 #include "IEditorUtilityExtension.h"
@@ -63,6 +64,7 @@ public:
 	bool IsLatestVersion() const;
 	bool AreSupportedClassesForBlueprints() const;
 	TArray<TSoftClassPtr<UObject>> GetSupportedClasses() const;
+	TArray<FAssetActionSupportCondition> GetAssetActionSupportConditions() const;
 	TArray<FBlutilityFunctionData> GetCallableFunctions() const;
 
 	const FAssetData& GetUtilityBlueprintAsset() const { return UtilityBlueprintAsset; }
@@ -80,6 +82,7 @@ public:
 public:
 	static void AddTagsFor_Version(TArray<UObject::FAssetRegistryTag>& OutTags);
 	static void AddTagsFor_SupportedClasses(const TArray<TSoftClassPtr<UObject>>& SupportedClasses, TArray<UObject::FAssetRegistryTag>& OutTags);
+	static void AddTagsFor_SupportedConditions(const TArray<FAssetActionSupportCondition>& SupportedConditions, TArray<UObject::FAssetRegistryTag>& OutTags);
 	static void AddTagsFor_IsActionForBlueprints(bool IsActionForBlueprints, TArray<UObject::FAssetRegistryTag>& OutTags);
 	static void AddTagsFor_CallableFunctions(const UObject* FunctionsSource, TArray<UObject::FAssetRegistryTag>& OutTags);
 
diff --git a/Engine/Source/Editor/ContentBrowser/Private/ContentBrowserMenuContexts.cpp b/Engine/Source/Editor/ContentBrowser/Private/ContentBrowserMenuContexts.cpp
--- a/Engine/Source/Editor/ContentBrowser/Private/ContentBrowserMenuContexts.cpp
+++ b/Engine/Source/Editor/ContentBrowser/Private/ContentBrowserMenuContexts.cpp
@@ -84,6 +84,16 @@ const FAssetData* UContentBrowserAssetContextMenuContext::GetSingleSelectedAsset
 
 namespace UE::ContentBrowser
 {
+	TArray<UToolMenu*> ExtendToolMenu_AssetContextMenus(TConstArrayView<UClass*> AssetClasses)
+	{
+		TArray<UToolMenu*> ToolMenus;
+		for (UClass* AssetClass : AssetClasses)
+		{
+			ToolMenus.Add(ExtendToolMenu_AssetContextMenu(TSoftClassPtr<UObject>(AssetClass)));
+		}
+		return ToolMenus;
+	}
+	
 	UToolMenu* ExtendToolMenu_AssetContextMenu(UClass* AssetClass)
 	{
 		return ExtendToolMenu_AssetContextMenu(TSoftClassPtr<UObject>(AssetClass));
diff --git a/Engine/Source/Editor/ContentBrowser/Public/ContentBrowserMenuContexts.h b/Engine/Source/Editor/ContentBrowser/Public/ContentBrowserMenuContexts.h
--- a/Engine/Source/Editor/ContentBrowser/Public/ContentBrowserMenuContexts.h
+++ b/Engine/Source/Editor/ContentBrowser/Public/ContentBrowserMenuContexts.h
@@ -269,4 +269,5 @@ namespace UE::ContentBrowser
 {
 	CONTENTBROWSER_API UToolMenu* ExtendToolMenu_AssetContextMenu(UClass* AssetClass);
 	CONTENTBROWSER_API UToolMenu* ExtendToolMenu_AssetContextMenu(TSoftClassPtr<UObject> AssetSoftClass);
+	CONTENTBROWSER_API TArray<UToolMenu*> ExtendToolMenu_AssetContextMenus(TConstArrayView<UClass*> AssetClasses);
 }
\ No newline at end of file
diff --git a/Engine/Source/Runtime/CoreUObject/Private/Misc/WarnIfAssetsLoadedInScope.cpp b/Engine/Source/Runtime/CoreUObject/Private/Misc/WarnIfAssetsLoadedInScope.cpp
--- a/Engine/Source/Runtime/CoreUObject/Private/Misc/WarnIfAssetsLoadedInScope.cpp
+++ b/Engine/Source/Runtime/CoreUObject/Private/Misc/WarnIfAssetsLoadedInScope.cpp
@@ -55,7 +55,7 @@ FWarnIfAssetsLoadedInScope::~FWarnIfAssetsLoadedInScope()
 
 void FWarnIfAssetsLoadedInScope::OnAssetLoaded(UObject* InObject)
 {
-	// Don't warn people if the switch isn't enalbed.
+	// Don't warn people if the switch isn't enabled.
 	if (!IsEnabled())
 	{
 		return;
