Update EditorUtilityWidget to only be enabled when we are not in PIE or in BP debugging.

diff --git a/Engine/Source/Editor/Blutility/Classes/EditorUtilityWidget.h b/Engine/Source/Editor/Blutility/Classes/EditorUtilityWidget.h
--- a/Engine/Source/Editor/Blutility/Classes/EditorUtilityWidget.h
+++ b/Engine/Source/Editor/Blutility/Classes/EditorUtilityWidget.h
@@ -23,7 +23,7 @@ class UObject;
 UCLASS(Abstract, meta = (ShowWorldContextPin), config = Editor)
 class BLUTILITY_API UEditorUtilityWidget : public UUserWidget
 {
-	GENERATED_UCLASS_BODY()
+	GENERATED_BODY()
 
 public:
 	// The default action called when the widget is invoked if bAutoRunDefaultAction=true (it is never called otherwise)
@@ -56,5 +56,4 @@ protected:
 	// Should this blueprint automatically run OnDefaultActionClicked, or should it open up a details panel to edit properties and/or offer multiple buttons
 	UPROPERTY(Category = Settings, EditDefaultsOnly, BlueprintReadOnly)
 	bool bAutoRunDefaultAction;
-
 };
diff --git a/Engine/Source/Editor/Blutility/Classes/EditorUtilityWidgetBlueprint.h b/Engine/Source/Editor/Blutility/Classes/EditorUtilityWidgetBlueprint.h
--- a/Engine/Source/Editor/Blutility/Classes/EditorUtilityWidgetBlueprint.h
+++ b/Engine/Source/Editor/Blutility/Classes/EditorUtilityWidgetBlueprint.h
@@ -36,7 +36,7 @@ enum class EMapChangeType : uint8;
 UCLASS()
 class BLUTILITY_API UEditorUtilityWidgetBlueprint : public UWidgetBlueprint
 {
-	GENERATED_UCLASS_BODY()
+	GENERATED_BODY()
 
 public:
 	virtual void BeginDestroy() override;
@@ -71,9 +71,19 @@ public:
 		return RegistrationName;
 	}
 
+private:
+	bool IsWidgetEnabled() const;
 	void ChangeTabWorld(UWorld* World, EMapChangeType MapChangeType);
 
 private:
+	// Should the widget be enabled when running PIE
+	UPROPERTY(Category = Settings, EditDefaultsOnly)
+	bool bIsEnabledInPIE = false;
+
+	// Should the widget be enabled when debugging BP
+	UPROPERTY(Category = Settings, EditDefaultsOnly)
+	bool bIsEnabledInDebugging = false;
+
 	FName RegistrationName;
 
 	TWeakPtr<SDockTab> CreatedTab;
diff --git a/Engine/Source/Editor/Blutility/Private/EditorUtilityWidget.cpp b/Engine/Source/Editor/Blutility/Private/EditorUtilityWidget.cpp
--- a/Engine/Source/Editor/Blutility/Private/EditorUtilityWidget.cpp
+++ b/Engine/Source/Editor/Blutility/Private/EditorUtilityWidget.cpp
@@ -10,11 +10,6 @@
 /////////////////////////////////////////////////////
 #define LOCTEXT_NAMESPACE "EditorUtility"
 
-UEditorUtilityWidget::UEditorUtilityWidget(const FObjectInitializer& ObjectInitializer)
-	: Super(ObjectInitializer)
-{
-}
-
 void UEditorUtilityWidget::ExecuteDefaultAction()
 {
 	check(bAutoRunDefaultAction);
diff --git a/Engine/Source/Editor/Blutility/Private/EditorUtilityWidgetBlueprint.cpp b/Engine/Source/Editor/Blutility/Private/EditorUtilityWidgetBlueprint.cpp
--- a/Engine/Source/Editor/Blutility/Private/EditorUtilityWidgetBlueprint.cpp
+++ b/Engine/Source/Editor/Blutility/Private/EditorUtilityWidgetBlueprint.cpp
@@ -38,12 +38,6 @@ class UWorld;
 /////////////////////////////////////////////////////
 // UEditorUtilityWidgetBlueprint
 
-UEditorUtilityWidgetBlueprint::UEditorUtilityWidgetBlueprint(const FObjectInitializer& ObjectInitializer)
-	: Super(ObjectInitializer)
-{
-
-}
-
 void UEditorUtilityWidgetBlueprint::BeginDestroy()
 {
 	// Only cleanup script if it has been registered and we're not shutdowning editor
@@ -94,13 +88,15 @@ TSharedRef<SWidget> UEditorUtilityWidgetBlueprint::CreateUtilityWidget()
 
 	UClass* BlueprintClass = GeneratedClass;
 	TSubclassOf<UEditorUtilityWidget> WidgetClass = BlueprintClass;
-	UWorld* World = GEditor->GetEditorWorldContext().World();
-	if (World)
+
+	if (CreatedUMGWidget)
+	{
+		CreatedUMGWidget->Rename(nullptr, GetTransientPackage(), REN_DoNotDirty);
+	}
+
+	CreatedUMGWidget = nullptr;
+	if (UWorld* World = GEditor->GetEditorWorldContext().World())
 	{
-		if (CreatedUMGWidget)
-		{
-			CreatedUMGWidget->Rename(nullptr, GetTransientPackage(), REN_DoNotDirty);
-		}
 		CreatedUMGWidget = CreateWidget<UEditorUtilityWidget>(World, WidgetClass);
 		if (CreatedUMGWidget)
 		{
@@ -125,6 +121,7 @@ TSharedRef<SWidget> UEditorUtilityWidgetBlueprint::CreateUtilityWidget()
 	if (CreatedUMGWidget)
 	{
 		TabWidget = SNew(SVerticalBox)
+			.IsEnabled_UObject(this, &UEditorUtilityWidgetBlueprint::IsWidgetEnabled)
 			+ SVerticalBox::Slot()
 			.HAlign(HAlign_Fill)
 			[
@@ -134,12 +131,19 @@ TSharedRef<SWidget> UEditorUtilityWidgetBlueprint::CreateUtilityWidget()
 	return TabWidget;
 }
 
+bool UEditorUtilityWidgetBlueprint::IsWidgetEnabled() const
+{
+	bool bAllowPIE = bIsEnabledInPIE || (GEditor && GEditor->PlayWorld == nullptr);
+	bool bAllowDebugging = bIsEnabledInDebugging || !GIntraFrameDebuggingGameThread;
+	return bAllowPIE && bAllowDebugging;
+}
+
 void UEditorUtilityWidgetBlueprint::RegenerateCreatedTab(UBlueprint* RecompiledBlueprint)
 {
-	if (CreatedTab.IsValid())
+	if (TSharedPtr<SDockTab> CreatedTabPinned = CreatedTab.Pin())
 	{
 		TSharedRef<SWidget> TabWidget = CreateUtilityWidget();
-		CreatedTab.Pin()->SetContent(TabWidget);
+		CreatedTabPinned->SetContent(TabWidget);
 	}
 }
 
@@ -177,7 +181,12 @@ void UEditorUtilityWidgetBlueprint::UpdateRespawnListIfNeeded(TSharedRef<SDockTa
 			BlutilityModule->RemoveLoadedScriptUI(this);
 		}
 	}
-	CreatedUMGWidget = nullptr;
+
+	if (CreatedUMGWidget)
+	{
+		CreatedUMGWidget->Rename(nullptr, GetTransientPackage());
+		CreatedUMGWidget = nullptr;
+	}
 }
 
 void UEditorUtilityWidgetBlueprint::GetReparentingRules(TSet< const UClass* >& AllowedChildrenOfClasses, TSet< const UClass* >& DisallowedChildrenOfClasses) const
