Add option to set an EUW tab display name & add an UnregisterTabByID method
Relates to review #23371479

diff --git a/Engine/Source/Editor/Blutility/Classes/EditorUtilityWidget.h b/Engine/Source/Editor/Blutility/Classes/EditorUtilityWidget.h
--- a/Engine/Source/Editor/Blutility/Classes/EditorUtilityWidget.h
+++ b/Engine/Source/Editor/Blutility/Classes/EditorUtilityWidget.h
@@ -43,9 +43,19 @@ public:
 		return bAutoRunDefaultAction;
 	}
 
+	/** Returns the default desired tab display name that was specified for this widget */
+	FText GetTabDisplayName() const
+	{
+		return TabDisplayName;
+	}
+
 	virtual bool IsEditorUtility() const override { return true; }
 
 protected:
+	/** The display name for tabs spawned with this widget */
+	UPROPERTY(Category = Config, EditDefaultsOnly, BlueprintReadWrite, AssetRegistrySearchable)
+	FText TabDisplayName;
+
 	UPROPERTY(Category = Config, EditDefaultsOnly, BlueprintReadWrite, AssetRegistrySearchable)
 	FString HelpText;
 
diff --git a/Engine/Source/Editor/Blutility/Classes/EditorUtilityWidgetBlueprint.h b/Engine/Source/Editor/Blutility/Classes/EditorUtilityWidgetBlueprint.h
--- a/Engine/Source/Editor/Blutility/Classes/EditorUtilityWidgetBlueprint.h
+++ b/Engine/Source/Editor/Blutility/Classes/EditorUtilityWidgetBlueprint.h
@@ -71,8 +71,12 @@ public:
 		return RegistrationName;
 	}
 
+	/** Returns the default desired tab display name that was specified for this widget */
+	FText GetTabDisplayName() const;
+
 private:
 	bool IsWidgetEnabled() const;
+
 	void ChangeTabWorld(UWorld* World, EMapChangeType MapChangeType);
 
 private:
diff --git a/Engine/Source/Editor/Blutility/Private/BlutilityModule.cpp b/Engine/Source/Editor/Blutility/Private/BlutilityModule.cpp
--- a/Engine/Source/Editor/Blutility/Private/BlutilityModule.cpp
+++ b/Engine/Source/Editor/Blutility/Private/BlutilityModule.cpp
@@ -138,7 +138,7 @@ public:
 						const UEditorUtilityWidget* CDO = Blueprint->GeneratedClass->GetDefaultObject<UEditorUtilityWidget>();
 						FName RegistrationName = FName(*(Blueprint->GetPathName() + LOCTEXT("ActiveTabSuffix", "_ActiveTab").ToString()));
 						Blueprint->SetRegistrationName(RegistrationName);
-						FText DisplayName = FText::FromString(FName::NameToDisplayString(Blueprint->GetName(), false));
+						FText DisplayName = Blueprint->GetTabDisplayName();
 						if (LevelEditorTabManager && !LevelEditorTabManager->HasTabSpawner(RegistrationName))
 						{
 							LevelEditorTabManager->RegisterTabSpawner(RegistrationName, FOnSpawnTab::CreateUObject(Blueprint, &UEditorUtilityWidgetBlueprint::SpawnEditorUITab))
diff --git a/Engine/Source/Editor/Blutility/Private/EditorUtilitySubsystem.cpp b/Engine/Source/Editor/Blutility/Private/EditorUtilitySubsystem.cpp
--- a/Engine/Source/Editor/Blutility/Private/EditorUtilitySubsystem.cpp
+++ b/Engine/Source/Editor/Blutility/Private/EditorUtilitySubsystem.cpp
@@ -224,7 +224,7 @@ void UEditorUtilitySubsystem::RegisterTabAndGetID(class UEditorUtilityWidgetBlue
 	if (InBlueprint && !IsRunningCommandlet())
 	{
 		FName RegistrationName = NewTabID.IsNone() ? FName(*(InBlueprint->GetPathName() + LOCTEXT("ActiveTabSuffix", "_ActiveTab").ToString())) : FName(*(InBlueprint->GetPathName() + NewTabID.ToString()));
-		FText DisplayName = FText::FromString(FName::NameToDisplayString(InBlueprint->GetName(), false));
+		FText DisplayName = InBlueprint->GetTabDisplayName();
 		FLevelEditorModule& LevelEditorModule = FModuleManager::GetModuleChecked<FLevelEditorModule>(TEXT("LevelEditor"));
 		if (TSharedPtr<FTabManager> LevelEditorTabManager = LevelEditorModule.GetLevelEditorTabManager())
 		{
@@ -253,10 +253,13 @@ bool UEditorUtilitySubsystem::SpawnRegisteredTabByID(FName NewTabID)
 			{
 				TSharedPtr<SDockTab> NewDockTab = LevelEditorTabManager->TryInvokeTab(NewTabID);
 				IBlutilityModule* BlutilityModule = FModuleManager::GetModulePtr<IBlutilityModule>("Blutility");
-				UEditorUtilityWidgetBlueprint* WidgetToSpawn = *RegisteredTabs.Find(NewTabID);
-				check(WidgetToSpawn);
-				BlutilityModule->AddLoadedScriptUI(WidgetToSpawn);
-				return true;
+				UEditorUtilityWidgetBlueprint** WidgetToSpawn = RegisteredTabs.Find(NewTabID);
+				if (WidgetToSpawn)
+				{
+					check(*WidgetToSpawn);
+					BlutilityModule->AddLoadedScriptUI(*WidgetToSpawn);
+					return true;
+				}
 			}
 		}
 	}
@@ -300,6 +303,29 @@ bool UEditorUtilitySubsystem::CloseTabByID(FName NewTabID)
 	return false;
 }
 
+bool UEditorUtilitySubsystem::UnregisterTabByID(FName TabID)
+{
+	if (!IsRunningCommandlet())
+	{
+		const FLevelEditorModule& LevelEditorModule = FModuleManager::GetModuleChecked<FLevelEditorModule>(TEXT("LevelEditor"));
+		const TSharedPtr<FTabManager> LevelEditorTabManager = LevelEditorModule.GetLevelEditorTabManager();
+		const TSharedPtr<SDockTab> FoundTab = LevelEditorTabManager->FindExistingLiveTab(TabID);
+		if (FoundTab)
+		{
+			FoundTab->RequestCloseTab();
+		}
+
+		if (RegisteredTabs.Contains(TabID) && LevelEditorTabManager->HasTabSpawner(TabID))
+		{
+			RegisteredTabs.Remove(TabID);
+			LevelEditorTabManager->UnregisterTabSpawner(TabID);
+			return true;
+		}
+	}
+
+	return false;
+}
+
 UEditorUtilityWidget* UEditorUtilitySubsystem::FindUtilityWidgetFromBlueprint(class UEditorUtilityWidgetBlueprint* InBlueprint)
 {
 	if (!IsValid(InBlueprint))
diff --git a/Engine/Source/Editor/Blutility/Private/EditorUtilityWidgetBlueprint.cpp b/Engine/Source/Editor/Blutility/Private/EditorUtilityWidgetBlueprint.cpp
--- a/Engine/Source/Editor/Blutility/Private/EditorUtilityWidgetBlueprint.cpp
+++ b/Engine/Source/Editor/Blutility/Private/EditorUtilityWidgetBlueprint.cpp
@@ -147,6 +147,18 @@ void UEditorUtilityWidgetBlueprint::RegenerateCreatedTab(UBlueprint* RecompiledB
 	}
 }
 
+
+FText UEditorUtilityWidgetBlueprint::GetTabDisplayName() const
+{
+	const UEditorUtilityWidget* EditorUtilityWidget = GeneratedClass->GetDefaultObject<UEditorUtilityWidget>();
+
+	if (EditorUtilityWidget && !EditorUtilityWidget->GetTabDisplayName().IsEmpty())
+	{
+		return EditorUtilityWidget->GetTabDisplayName();
+	}
+	return FText::FromString(FName::NameToDisplayString(GetName(), false));
+}
+
 void UEditorUtilityWidgetBlueprint::ChangeTabWorld(UWorld* World, EMapChangeType MapChangeType)
 {
 	if (MapChangeType == EMapChangeType::TearDownWorld)
diff --git a/Engine/Source/Editor/Blutility/Public/EditorUtilitySubsystem.h b/Engine/Source/Editor/Blutility/Public/EditorUtilitySubsystem.h
--- a/Engine/Source/Editor/Blutility/Public/EditorUtilitySubsystem.h
+++ b/Engine/Source/Editor/Blutility/Public/EditorUtilitySubsystem.h
@@ -96,6 +96,10 @@ public:
 	UFUNCTION(BlueprintCallable, Category = "Development|Editor")
 	bool CloseTabByID(FName NewTabID);
 
+	/** Given an ID for a tab, try to close and unregister a tab that was registered through this subsystem */
+	UFUNCTION(BlueprintCallable, Category = "Development|Editor")
+	bool UnregisterTabByID(FName TabID);
+
 	/** Given an editor utility widget blueprint, get the widget it creates. This will return a null pointer if the widget is not currently in a tab.*/
 	UFUNCTION(BlueprintCallable, Category = "Development|Editor")
 	UEditorUtilityWidget* FindUtilityWidgetFromBlueprint(class UEditorUtilityWidgetBlueprint* InBlueprint);
