Exposed IContentBrowserSingleton::GetSelectedPathViewFolders to Blueprints and Python via the EditorUtilityLibrary

diff --git a/Engine/Source/Editor/Blutility/Private/EditorUtilityLibrary.cpp b/Engine/Source/Editor/Blutility/Private/EditorUtilityLibrary.cpp
--- a/Engine/Source/Editor/Blutility/Private/EditorUtilityLibrary.cpp
+++ b/Engine/Source/Editor/Blutility/Private/EditorUtilityLibrary.cpp
@@ -335,6 +335,14 @@ TArray<FString> UEditorUtilityLibrary::GetSelectedFolderPaths()
 	return Paths;
 }
 
+TArray<FString> UEditorUtilityLibrary::GetSelectedPathViewFolderPaths()
+{
+	IContentBrowserSingleton& ContentBrowser = FModuleManager::LoadModuleChecked<FContentBrowserModule>("ContentBrowser").Get();
+	TArray<FString> Paths;
+    ContentBrowser.GetSelectedPathViewFolders(Paths);
+	return Paths;
+}
+
 void UEditorUtilityLibrary::SyncBrowserToFolders(const TArray<FString>& FolderList)
 {
 	FContentBrowserModule& ContentBrowserModule = FModuleManager::Get().LoadModuleChecked<FContentBrowserModule>("ContentBrowser");
diff --git a/Engine/Source/Editor/Blutility/Public/EditorUtilityLibrary.h b/Engine/Source/Editor/Blutility/Public/EditorUtilityLibrary.h
--- a/Engine/Source/Editor/Blutility/Public/EditorUtilityLibrary.h
+++ b/Engine/Source/Editor/Blutility/Public/EditorUtilityLibrary.h
@@ -183,6 +183,10 @@ public:
 	// Gets the path to the currently selected folder in the content browser
 	UFUNCTION(BlueprintPure, Category = "Development|Editor")
 	static TArray<FString> GetSelectedFolderPaths();
+	
+	// Returns the folders that are selected in the path view for the content browser
+    UFUNCTION(BlueprintPure, Category = "Development|Editor")
+    static TArray<FString> GetSelectedPathViewFolderPaths();
 
 	/**
 	 * Sync the Content Browser to the given folder(s)
