Add UEditorUtilityLibrary::GetCurrentContentBrowserItemPath()

- Alternative to GetCurrentContentBrowserPath() which returns a string only for internal paths (not virtual paths)

diff --git a/Engine/Source/Editor/Blutility/Private/EditorUtilityLibrary.cpp b/Engine/Source/Editor/Blutility/Private/EditorUtilityLibrary.cpp
--- a/Engine/Source/Editor/Blutility/Private/EditorUtilityLibrary.cpp
+++ b/Engine/Source/Editor/Blutility/Private/EditorUtilityLibrary.cpp
@@ -327,6 +327,12 @@ bool UEditorUtilityLibrary::GetCurrentContentBrowserPath(FString& OutPath)
 	}
 }
 
+FContentBrowserItemPath UEditorUtilityLibrary::GetCurrentContentBrowserItemPath()
+{
+	IContentBrowserSingleton& ContentBrowser = FModuleManager::LoadModuleChecked<FContentBrowserModule>("ContentBrowser").Get();
+	return ContentBrowser.GetCurrentPath();
+}
+
 TArray<FString> UEditorUtilityLibrary::GetSelectedFolderPaths()
 {
 	IContentBrowserSingleton& ContentBrowser = FModuleManager::LoadModuleChecked<FContentBrowserModule>("ContentBrowser").Get();
diff --git a/Engine/Source/Editor/Blutility/Public/EditorUtilityLibrary.h b/Engine/Source/Editor/Blutility/Public/EditorUtilityLibrary.h
--- a/Engine/Source/Editor/Blutility/Public/EditorUtilityLibrary.h
+++ b/Engine/Source/Editor/Blutility/Public/EditorUtilityLibrary.h
@@ -23,6 +23,7 @@ class UClass;
 class UEditorPerProjectUserSettings;
 class UWorld;
 struct FAssetData;
+struct FContentBrowserItemPath;
 struct FFrame;
 
 UCLASS()
@@ -174,12 +175,17 @@ public:
 
 	/**
 	 * Attempts to get the path for the active content browser, returns false if there is no active content browser
+	 * or if it was a virtual path
 	 * @param	OutPath	The returned path if successfully found
 	 * @return	Whether a path was successfully returned
 	 */
 	UFUNCTION(BlueprintPure, Category = "Development|Editor")
 	static bool GetCurrentContentBrowserPath(FString& OutPath);
 
+	// Gets the current content browser path if one is open, whether it is internal or virtual.
+	UFUNCTION(BlueprintPure, Category = "Development|Editor")
+	static FContentBrowserItemPath GetCurrentContentBrowserItemPath();
+
 	// Gets the path to the currently selected folder in the content browser
 	UFUNCTION(BlueprintPure, Category = "Development|Editor")
 	static TArray<FString> GetSelectedFolderPaths();
