Fix to a crash when using an Editor Blueprint as the root of another Editor Widget Blueprint

diff --git a/Engine/Source/Editor/Blutility/Private/EditorUtilityWidgetBlueprint.cpp b/Engine/Source/Editor/Blutility/Private/EditorUtilityWidgetBlueprint.cpp
--- a/Engine/Source/Editor/Blutility/Private/EditorUtilityWidgetBlueprint.cpp
+++ b/Engine/Source/Editor/Blutility/Private/EditorUtilityWidgetBlueprint.cpp
@@ -112,7 +112,10 @@ TSharedRef<SWidget> UEditorUtilityWidgetBlueprint::CreateUtilityWidget()
 				if (Widget->IsA(UEditorUtilityWidget::StaticClass()))
 				{
 					Widget->SetFlags(RF_Transient);
-					Widget->Slot->SetFlags(RF_Transient);
+					if (UPanelSlot* Slot = Widget->Slot)
+					{
+						Slot->SetFlags(RF_Transient);
+					}
 				}
 			}
 		}
