PR #9978: [Blutility] Fix classes included in other modules  (Contributed by [REDACTED_PERSON])

diff --git a/Engine/Source/Editor/Blutility/Private/IEditorUtilityExtension.h b/Engine/Source/Editor/Blutility/Private/IEditorUtilityExtension.h
deleted file mode 100644
diff --git a/Engine/Source/Editor/Blutility/Public/IEditorUtilityExtension.h b/Engine/Source/Editor/Blutility/Public/IEditorUtilityExtension.h
new file mode 100644
--- /dev/null
+++ b/Engine/Source/Editor/Blutility/Public/IEditorUtilityExtension.h
@@ -0,0 +1,20 @@
+// Copyright Epic Games, Inc. All Rights Reserved.
+
+#pragma once
+
+#include "CoreMinimal.h"
+#include "UObject/ObjectMacros.h"
+#include "UObject/Interface.h"
+
+#include "IEditorUtilityExtension.generated.h"
+
+UINTERFACE(BlueprintType)
+class BLUTILITY_API UEditorUtilityExtension : public UInterface
+{
+	GENERATED_UINTERFACE_BODY()
+};
+
+class BLUTILITY_API IEditorUtilityExtension
+{
+	GENERATED_IINTERFACE_BODY()
+};
