Add function to ConvertToEditorUtilityWidget, just meant to be a way to salvage assets that are violating content policies (e.g. calling editor only/cook only routines)

diff --git a/Engine/Plugins/Editor/EngineAssetDefinitions/Source/Private/AssetDefinition_EditorUtilityWidgetBlueprint.cpp b/Engine/Plugins/Editor/EngineAssetDefinitions/Source/Private/AssetDefinition_EditorUtilityWidgetBlueprint.cpp
--- a/Engine/Plugins/Editor/EngineAssetDefinitions/Source/Private/AssetDefinition_EditorUtilityWidgetBlueprint.cpp
+++ b/Engine/Plugins/Editor/EngineAssetDefinitions/Source/Private/AssetDefinition_EditorUtilityWidgetBlueprint.cpp
@@ -4,6 +4,7 @@
 
 #include "Algo/RemoveIf.h"
 #include "ContentBrowserMenuContexts.h"
+#include "EditorUtilityLibrary.h"
 #include "EditorUtilitySubsystem.h"
 #include "EditorUtilityWidgetBlueprint.h"
 #include "IBlutilityModule.h"
@@ -98,6 +99,18 @@ namespace MenuExtension_EditorUtilityWidgetBlueprint
 			}
 		}
 	}
+
+	void ExecuteConvertToEditorUtilityBlueprint(const FToolMenuContext& InContext, TArray<FAssetData> AssetData)
+	{
+		if (const UContentBrowserAssetContextMenuContext* Context = UContentBrowserAssetContextMenuContext::FindContextWithAssets(InContext))
+		{
+			TArray<UWidgetBlueprint*> WidgetBlueprints = Context->LoadSelectedObjects<UWidgetBlueprint>();
+			for(UWidgetBlueprint* WidgetBP : WidgetBlueprints)
+			{
+				UEditorUtilityLibrary::ConvertToEditorUtilityWidget(WidgetBP);
+			}
+		}
+	}
 	
 	FDelayedAutoRegisterHelper DelayedAutoRegister(EDelayedRegisterRunPhase::EndOfEngineInit,[]
 	{
@@ -132,6 +145,35 @@ namespace MenuExtension_EditorUtilityWidgetBlueprint
 					}
 				}
 			}));
+
+			UToolMenu* WigetBPMenu = UToolMenus::Get()->ExtendMenu("ContentBrowser.AssetContextMenu.WidgetBlueprint.AssetActionsSubMenu");
+
+			FToolMenuSection& WidgetBPSection = WigetBPMenu->FindOrAddSection("AssetContextAdvancedActions");
+			WidgetBPSection.AddDynamicEntry(NAME_None, FNewToolMenuSectionDelegate::CreateLambda([](FToolMenuSection& InSection)
+				{
+					if (const UContentBrowserAssetContextMenuContext* Context = UContentBrowserAssetContextMenuContext::FindContextWithAssets(InSection))
+					{
+						TArray<FAssetData> SelectedAssets(Context->SelectedAssets);
+
+						const int32 EndOfRange = Algo::StableRemoveIf(SelectedAssets, [](const FAssetData& AssetData)
+							{
+								return !AssetData.IsInstanceOf(UWidgetBlueprint::StaticClass());
+							});
+
+						SelectedAssets.SetNum(EndOfRange);
+
+						if (!SelectedAssets.IsEmpty())
+						{
+							const TAttribute<FText> Label = LOCTEXT("EditorUtilityWidget_Convert", "Convert to Editor Utility Widget");
+							const TAttribute<FText> ToolTip = LOCTEXT("ConvertTooltip", "Convert this Widget Blueprint to an Editor Utility Widget");
+							const FSlateIcon Icon = FSlateIcon();
+
+							FToolUIAction UIAction;
+							UIAction.ExecuteAction = FToolMenuExecuteAction::CreateStatic(&ExecuteConvertToEditorUtilityBlueprint, MoveTemp(SelectedAssets));
+							InSection.AddMenuEntry("EditorUtility_Convert", Label, ToolTip, Icon, UIAction);
+						}
+					}
+				}));
 		}));
 	});
 }
diff --git a/Engine/Source/Editor/Blutility/Private/EditorUtilityLibrary.cpp b/Engine/Source/Editor/Blutility/Private/EditorUtilityLibrary.cpp
--- a/Engine/Source/Editor/Blutility/Private/EditorUtilityLibrary.cpp
+++ b/Engine/Source/Editor/Blutility/Private/EditorUtilityLibrary.cpp
@@ -1,6 +1,8 @@
 // Copyright Epic Games, Inc. All Rights Reserved.
 
 #include "EditorUtilityLibrary.h"
+#include "Blueprint/WidgetTree.h"
+#include "EditorUtilityWidgetBlueprint.h"
 #include "Engine/Selection.h"
 #include "Editor.h"
 #include "GameFramework/Actor.h"
@@ -11,7 +13,10 @@
 #include "IAssetTools.h"
 #include "EditorUtilitySubsystem.h"
 #include "Kismet/KismetSystemLibrary.h"
+#include "Serialization/ArchiveReplaceObjectRef.h"
+#include "Serialization/FindReferencersArchive.h"
 #include "Templates/SubclassOf.h"
+#include "WidgetBlueprint.h"
 
 #define LOCTEXT_NAMESPACE "BlutilityLevelEditorExtensions"
 
@@ -345,6 +350,75 @@ void UEditorUtilityLibrary::SyncBrowserToFolders(const TArray<FString>& FolderLi
 	ContentBrowserModule.Get().SyncBrowserToFolders( FolderList, false, true );
 }
 
+void UEditorUtilityLibrary::ConvertToEditorUtilityWidget(UWidgetBlueprint* WidgetBP)
+{
+	if (!WidgetBP)
+	{
+		return;
+	}
+
+	if (WidgetBP->IsA<UEditorUtilityWidgetBlueprint>())
+	{
+		return;
+	}
+
+	FName BPName = WidgetBP->GetFName();
+	UObject* Outer = WidgetBP->GetOuter();
+	EObjectFlags Flags = WidgetBP->GetFlags();
+
+	// Rename the blueprint out of the way, create a new EUWBP, and then
+	// put all of the blueprints child objects back 'under it'. The Blueprint
+	// generated calss does not require any updating.
+
+	TArray<struct FEditedDocumentInfo> OriginalEditedDocuments = WidgetBP->LastEditedDocuments;
+	WidgetBP->Rename(nullptr, GetTransientPackage(), REN_DontCreateRedirectors|REN_SkipGeneratedClasses|REN_ForceNoResetLoaders);
+	TArray<UObject*> Children;
+	GetObjectsWithOuter(WidgetBP, Children, false);
+
+	UEditorUtilityWidgetBlueprint* EWBP = NewObject<UEditorUtilityWidgetBlueprint>(Outer, BPName, Flags);
+	if (EWBP->WidgetTree)
+	{
+		// WidgetTree is a DSO created as a side effect of construction,
+		// we just want to use the existing one:
+		EWBP->WidgetTree->Rename(nullptr, GetTransientPackage(), REN_ForceNoResetLoaders | REN_DontCreateRedirectors);
+		EWBP->WidgetTree = WidgetBP->WidgetTree;
+	}
+	for (UObject* Child : Children)
+	{
+		Child->Rename(nullptr, EWBP, REN_DontCreateRedirectors | REN_SkipGeneratedClasses | REN_ForceNoResetLoaders);
+	}
+
+	UEngine::FCopyPropertiesForUnrelatedObjectsParams Params;
+	Params.bPerformDuplication = true;
+	Params.bNotifyObjectReplacement = false;
+	Params.bPreserveRootComponent = false;
+	UEngine::CopyPropertiesForUnrelatedObjects(WidgetBP, EWBP);
+	EWBP->LastEditedDocuments = OriginalEditedDocuments; // mangled by UBlueprint::Rename.. unmangle
+
+	EWBP->GeneratedClass->ClassGeneratedBy = EWBP; // update ClassGeneratedBy on the UClass
+	check(EWBP->GeneratedClass == nullptr || EWBP->GeneratedClass->GetOuter() == EWBP->GetOuter());
+
+	TMap<UObject*, UObject*> OldToNew;
+	OldToNew.Add(WidgetBP, EWBP);
+
+	// Update any references to the UBlueprint itself:
+
+	TArray<UObject*> Targets = { WidgetBP, GetTransientPackage() };
+	FArchiveReplaceObjectRef<UObject> ReplaceReferencesInRoot(EWBP, OldToNew);
+	FFindReferencersArchive Archive(EWBP, Targets);
+	check(Archive.GetReferenceCount(WidgetBP) == 0);
+
+	Children.Reset();
+	GetObjectsWithOuter(EWBP->GetOuter(), Children);
+
+	for (UObject* Child : Children)
+	{
+		FArchiveReplaceObjectRef<UObject> ReplaceReferences(Child, OldToNew);
+		FFindReferencersArchive Archive2(Child, Targets);
+		check(Archive2.GetReferenceCount(WidgetBP) == 0);
+	}
+}
+
 #endif
 
 #undef LOCTEXT_NAMESPACE
diff --git a/Engine/Source/Editor/Blutility/Public/EditorUtilityLibrary.h b/Engine/Source/Editor/Blutility/Public/EditorUtilityLibrary.h
--- a/Engine/Source/Editor/Blutility/Public/EditorUtilityLibrary.h
+++ b/Engine/Source/Editor/Blutility/Public/EditorUtilityLibrary.h
@@ -200,5 +200,8 @@ public:
 	 */
 	UFUNCTION(BlueprintCallable, Category = "Editor Scripting | Content Browser")
 	static void SyncBrowserToFolders(const TArray<FString>& FolderList);
+	
+	UFUNCTION(BlueprintCallable, Category = "Editor Scripting | Content Browser")
+	static void ConvertToEditorUtilityWidget(class UWidgetBlueprint* WidgetBP);
 #endif
 };
