Add Editor Utility version of UMG widget components to prevent serialization conflicts

diff --git a/Engine/Config/BaseEditor.ini b/Engine/Config/BaseEditor.ini
--- a/Engine/Config/BaseEditor.ini
+++ b/Engine/Config/BaseEditor.ini
@@ -643,3 +643,22 @@ bShowWidgetsFromDeveloperContent=True
 bUseWidgetTemplateSelector=True
 +CommonRootWidgetClasses=/Script/UMG.StackBox
 +CommonRootWidgetClasses=/Script/UMG.GridPanel
++WidgetClassesToHide=/Script/UMG.Button
++WidgetClassesToHide=/Script/UMG.CheckBox
++WidgetClassesToHide=/Script/UMG.CircularThrobber
++WidgetClassesToHide=/Script/UMG.ComboBoxKey
++WidgetClassesToHide=/Script/UMG.ComboBoxString
++WidgetClassesToHide=/Script/UMG.EditableText
++WidgetClassesToHide=/Script/UMG.EditableTextBox
++WidgetClassesToHide=/Script/UMG.ExpandableArea
++WidgetClassesToHide=/Script/UMG.InputKeySelector
++WidgetClassesToHide=/Script/UMG.ListView
++WidgetClassesToHide=/Script/UMG.MultiLineEditableText
++WidgetClassesToHide=/Script/UMG.MultiLineEditableTextBox
++WidgetClassesToHide=/Script/UMG.ProgressBar
++WidgetClassesToHide=/Script/UMG.ScrollBar
++WidgetClassesToHide=/Script/UMG.ScrollBox
++WidgetClassesToHide=/Script/UMG.Slider
++WidgetClassesToHide=/Script/UMG.SpinBox
++WidgetClassesToHide=/Script/UMG.Throbber
++WidgetClassesToHide=/Script/UMG.TreeView
\ No newline at end of file
diff --git a/Engine/Source/Editor/Blutility/Classes/EditorUtilityWidgetComponents.h b/Engine/Source/Editor/Blutility/Classes/EditorUtilityWidgetComponents.h
new file mode 100644
--- /dev/null
+++ b/Engine/Source/Editor/Blutility/Classes/EditorUtilityWidgetComponents.h
@@ -0,0 +1,150 @@
+// Copyright Epic Games, Inc. All Rights Reserved.
+
+/**
+ * Editor Utility Specfic Widget Components
+ * 
+ * These exist to provide a UE5 style for Widget Blueprints. Historically
+ * we conditionally changed styling in constructor to achive this style
+ * however that causes issues with CDO comparision.
+ */
+
+#pragma once
+
+#include "Components/Button.h"
+#include "Components/CheckBox.h"
+#include "Components/CircularThrobber.h"
+#include "Components/ComboBoxKey.h"
+#include "Components/ComboBoxString.h"
+#include "Components/EditableText.h"
+#include "Components/EditableTextBox.h"
+#include "Components/ExpandableArea.h"
+#include "Components/InputKeySelector.h"
+#include "Components/ListView.h"
+#include "Components/MultiLineEditableText.h"
+#include "Components/MultiLineEditableTextBox.h"
+#include "Components/ProgressBar.h"
+#include "Components/ScrollBar.h"
+#include "Components/ScrollBox.h"
+#include "Components/Slider.h"
+#include "Components/SpinBox.h"
+#include "Components/Throbber.h"
+#include "Components/TreeView.h"
+
+#include "EditorUtilityWidgetComponents.generated.h"
+
+UCLASS()
+class BLUTILITY_API UEditorUtilityButton : public UButton
+{
+	GENERATED_UCLASS_BODY()
+};
+
+UCLASS()
+class BLUTILITY_API UEditorUtilityCheckBox : public UCheckBox
+{
+	GENERATED_UCLASS_BODY()
+};
+
+UCLASS()
+class BLUTILITY_API UEditorUtilityCircularThrobber : public UCircularThrobber
+{
+	GENERATED_UCLASS_BODY()
+};
+
+UCLASS()
+class BLUTILITY_API UEditorUtilityComboBoxKey : public UComboBoxKey
+{
+	GENERATED_BODY()
+	
+public:
+	UEditorUtilityComboBoxKey();
+};
+
+UCLASS()
+class BLUTILITY_API UEditorUtilityComboBoxString : public UComboBoxString
+{
+	GENERATED_UCLASS_BODY()
+};
+
+UCLASS()
+class BLUTILITY_API UEditorUtilityEditableText : public UEditableText
+{
+	GENERATED_UCLASS_BODY()
+};
+
+UCLASS()
+class BLUTILITY_API UEditorUtilityEditableTextBox : public UEditableTextBox
+{
+	GENERATED_UCLASS_BODY()
+};
+
+UCLASS()
+class BLUTILITY_API UEditorUtilityExpandableArea : public UExpandableArea
+{
+	GENERATED_UCLASS_BODY()
+};
+
+UCLASS()
+class BLUTILITY_API UEditorUtilityInputKeySelector : public UInputKeySelector
+{
+	GENERATED_UCLASS_BODY()
+};
+
+UCLASS()
+class BLUTILITY_API UEditorUtilityListView : public UListView
+{
+	GENERATED_UCLASS_BODY()
+};
+
+UCLASS()
+class BLUTILITY_API UEditorUtilityMultiLineEditableText : public UMultiLineEditableText
+{
+	GENERATED_UCLASS_BODY()
+};
+
+UCLASS()
+class BLUTILITY_API UEditorUtilityMultiLineEditableTextBox : public UMultiLineEditableTextBox
+{
+	GENERATED_UCLASS_BODY()
+};
+
+UCLASS()
+class BLUTILITY_API UEditorUtilityProgressBar : public UProgressBar
+{
+	GENERATED_UCLASS_BODY()
+};
+
+UCLASS()
+class BLUTILITY_API UEditorUtilityScrollBar : public UScrollBar
+{
+	GENERATED_UCLASS_BODY()
+};
+
+UCLASS()
+class BLUTILITY_API UEditorUtilityScrollBox : public UScrollBox
+{
+	GENERATED_UCLASS_BODY()
+};
+
+UCLASS()
+class BLUTILITY_API UEditorUtilitySlider : public USlider
+{
+	GENERATED_UCLASS_BODY()
+};
+
+UCLASS()
+class BLUTILITY_API UEditorUtilitySpinBox : public USpinBox
+{
+	GENERATED_UCLASS_BODY()
+};
+
+UCLASS()
+class BLUTILITY_API UEditorUtilityThrobber : public UThrobber
+{
+	GENERATED_UCLASS_BODY()
+};
+
+UCLASS()
+class BLUTILITY_API UEditorUtilityTreeView : public UTreeView
+{
+	GENERATED_UCLASS_BODY()
+};
diff --git a/Engine/Source/Editor/Blutility/Private/EditorUtilityWidgetComponents.cpp b/Engine/Source/Editor/Blutility/Private/EditorUtilityWidgetComponents.cpp
new file mode 100644
--- /dev/null
+++ b/Engine/Source/Editor/Blutility/Private/EditorUtilityWidgetComponents.cpp
@@ -0,0 +1,131 @@
+// Copyright Epic Games, Inc. All Rights Reserved.
+
+#include "EditorUtilityWidgetComponents.h"
+
+#include "Styling/DefaultStyleCache.h"
+#include "Styling/UMGCoreStyle.h"
+
+UEditorUtilityButton::UEditorUtilityButton(const FObjectInitializer& ObjectInitializer)
+	: Super(ObjectInitializer)
+{
+	SetStyle(UE::Slate::Private::FDefaultStyleCache::GetEditor().GetButtonStyle());
+}
+
+UEditorUtilityCheckBox::UEditorUtilityCheckBox(const FObjectInitializer& ObjectInitializer)
+	: Super(ObjectInitializer)
+{
+	SetWidgetStyle(UE::Slate::Private::FDefaultStyleCache::GetEditor().GetCheckboxStyle());
+}
+
+UEditorUtilityCircularThrobber::UEditorUtilityCircularThrobber(const FObjectInitializer& ObjectInitializer)
+	: Super(ObjectInitializer)
+{
+	SetImage(UE::Slate::Private::FDefaultStyleCache::GetEditor().GetCircularThrobberBrushStyle());
+}
+
+UEditorUtilityComboBoxKey::UEditorUtilityComboBoxKey()
+	: Super()
+{
+	SetWidgetStyle(UE::Slate::Private::FDefaultStyleCache::GetEditor().GetComboBoxStyle());
+	SetItemStyle(UE::Slate::Private::FDefaultStyleCache::GetEditor().GetTableRowStyle());
+	InitScrollBarStyle(UE::Slate::Private::FDefaultStyleCache::GetEditor().GetScrollBarStyle());
+}
+
+UEditorUtilityComboBoxString::UEditorUtilityComboBoxString(const FObjectInitializer& ObjectInitializer)
+	: Super(ObjectInitializer)
+{
+	SetWidgetStyle(UE::Slate::Private::FDefaultStyleCache::GetEditor().GetComboBoxStyle());
+	SetItemStyle(UE::Slate::Private::FDefaultStyleCache::GetEditor().GetTableRowStyle());
+	InitScrollBarStyle(UE::Slate::Private::FDefaultStyleCache::GetEditor().GetScrollBarStyle());
+}
+
+UEditorUtilityEditableText::UEditorUtilityEditableText(const FObjectInitializer& ObjectInitializer)
+	: Super(ObjectInitializer)
+{
+	SetWidgetStyle(UE::Slate::Private::FDefaultStyleCache::GetEditor().GetEditableTextStyle());
+}
+
+UEditorUtilityEditableTextBox::UEditorUtilityEditableTextBox(const FObjectInitializer& ObjectInitializer)
+	: Super(ObjectInitializer)
+{
+	WidgetStyle = UE::Slate::Private::FDefaultStyleCache::GetEditor().GetEditableTextBoxStyle();
+}
+
+UEditorUtilityExpandableArea::UEditorUtilityExpandableArea(const FObjectInitializer& ObjectInitializer)
+	: Super(ObjectInitializer)
+{
+	SetStyle(UE::Slate::Private::FDefaultStyleCache::GetEditor().GetExpandableAreaStyle());
+	SetBorderBrush(UE::Slate::Private::FDefaultStyleCache::GetEditor().GetExpandableAreaBorderBrush());
+}
+
+UEditorUtilityInputKeySelector::UEditorUtilityInputKeySelector(const FObjectInitializer& ObjectInitializer)
+	: Super(ObjectInitializer)
+{
+	SetButtonStyle(UE::Slate::Private::FDefaultStyleCache::GetEditor().GetButtonStyle());
+	SetTextStyle(UE::Slate::Private::FDefaultStyleCache::GetEditor().GetTextBlockStyle());
+}
+
+UEditorUtilityListView::UEditorUtilityListView(const FObjectInitializer& ObjectInitializer)
+	: Super(ObjectInitializer)
+{
+	WidgetStyle = UE::Slate::Private::FDefaultStyleCache::GetEditor().GetListViewStyle();
+	ScrollBarStyle = UE::Slate::Private::FDefaultStyleCache::GetEditor().GetScrollBarStyle();
+}
+
+UEditorUtilityMultiLineEditableText::UEditorUtilityMultiLineEditableText(const FObjectInitializer& ObjectInitializer)
+	: Super(ObjectInitializer)
+{
+	SetWidgetStyle(UE::Slate::Private::FDefaultStyleCache::GetEditor().GetTextBlockStyle());
+}
+
+UEditorUtilityMultiLineEditableTextBox::UEditorUtilityMultiLineEditableTextBox(const FObjectInitializer& ObjectInitializer)
+	: Super(ObjectInitializer)
+{
+	WidgetStyle = UE::Slate::Private::FDefaultStyleCache::GetEditor().GetEditableTextBoxStyle();
+PRAGMA_DISABLE_DEPRECATION_WARNINGS
+	TextStyle_DEPRECATED = WidgetStyle.TextStyle;
+PRAGMA_ENABLE_DEPRECATION_WARNINGS
+}
+
+UEditorUtilityProgressBar::UEditorUtilityProgressBar(const FObjectInitializer& ObjectInitializer)
+	: Super(ObjectInitializer)
+{
+	SetWidgetStyle(UE::Slate::Private::FDefaultStyleCache::GetEditor().GetProgressBarStyle());
+}
+
+UEditorUtilityScrollBar::UEditorUtilityScrollBar(const FObjectInitializer& ObjectInitializer)
+	: Super(ObjectInitializer)
+{
+	SetWidgetStyle(UE::Slate::Private::FDefaultStyleCache::GetEditor().GetScrollBarStyle());
+}
+
+UEditorUtilityScrollBox::UEditorUtilityScrollBox(const FObjectInitializer& ObjectInitializer)
+	: Super(ObjectInitializer)
+{
+	SetWidgetStyle(UE::Slate::Private::FDefaultStyleCache::GetEditor().GetScrollBoxStyle());
+	SetWidgetBarStyle(UE::Slate::Private::FDefaultStyleCache::GetEditor().GetScrollBarStyle());
+}
+
+UEditorUtilitySlider::UEditorUtilitySlider(const FObjectInitializer& ObjectInitializer)
+	: Super(ObjectInitializer)
+{
+	SetWidgetStyle(UE::Slate::Private::FDefaultStyleCache::GetEditor().GetSliderStyle());
+}
+
+UEditorUtilitySpinBox::UEditorUtilitySpinBox(const FObjectInitializer& ObjectInitializer)
+	: Super(ObjectInitializer)
+{
+	SetWidgetStyle(UE::Slate::Private::FDefaultStyleCache::GetEditor().GetSpinBoxStyle());
+}
+
+UEditorUtilityThrobber::UEditorUtilityThrobber(const FObjectInitializer& ObjectInitializer)
+	: Super(ObjectInitializer)
+{
+	SetImage(UE::Slate::Private::FDefaultStyleCache::GetEditor().GetThrobberBrush());
+}
+
+UEditorUtilityTreeView::UEditorUtilityTreeView(const FObjectInitializer& ObjectInitializer)
+	: Super(ObjectInitializer)
+{
+	WidgetStyle = UE::Slate::Private::FDefaultStyleCache::GetEditor().GetTreeViewStyle();
+}
