Editor Utility Widgets: Fix to potential crash when trying to open the same Editor Utility Widget more than once using Spawn and Register and Get ID function.

diff --git a/Engine/Source/Editor/Blutility/Private/EditorUtilitySubsystem.cpp b/Engine/Source/Editor/Blutility/Private/EditorUtilitySubsystem.cpp
--- a/Engine/Source/Editor/Blutility/Private/EditorUtilitySubsystem.cpp
+++ b/Engine/Source/Editor/Blutility/Private/EditorUtilitySubsystem.cpp
@@ -203,8 +203,10 @@ void UEditorUtilitySubsystem::ReleaseInstanceOfAsset(UObject* Asset)
 
 UEditorUtilityWidget* UEditorUtilitySubsystem::SpawnAndRegisterTabAndGetID(UEditorUtilityWidgetBlueprint* InBlueprint, FName& NewTabID)
 {
-	RegisterTabAndGetID(InBlueprint, NewTabID);
-	SpawnRegisteredTabByID(NewTabID);
+	FName TabID;
+	RegisterTabAndGetID(InBlueprint, TabID);
+	SpawnRegisteredTabByID(TabID);
+	NewTabID = TabID;
 	return FindUtilityWidgetFromBlueprint(InBlueprint);
 }
 
