Update utility widgets to mark their child utility widgets as transient recursively. This prevents an issue where the level gets marked dirty by nested utility widgets.

diff --git a/Engine/Source/Editor/Blutility/Classes/EditorUtilityWidgetBlueprint.h b/Engine/Source/Editor/Blutility/Classes/EditorUtilityWidgetBlueprint.h
--- a/Engine/Source/Editor/Blutility/Classes/EditorUtilityWidgetBlueprint.h
+++ b/Engine/Source/Editor/Blutility/Classes/EditorUtilityWidgetBlueprint.h
@@ -82,6 +82,8 @@ private:
 
 	void ChangeTabWorld(UWorld* World, EMapChangeType MapChangeType);
 
+	void MarkTransientRecursive(UEditorUtilityWidget* UtilityWidget);
+
 private:
 	// Should the widget be enabled when running PIE
 	UPROPERTY(Category = Settings, EditDefaultsOnly)
diff --git a/Engine/Source/Editor/Blutility/Private/EditorUtilityWidgetBlueprint.cpp b/Engine/Source/Editor/Blutility/Private/EditorUtilityWidgetBlueprint.cpp
--- a/Engine/Source/Editor/Blutility/Private/EditorUtilityWidgetBlueprint.cpp
+++ b/Engine/Source/Editor/Blutility/Private/EditorUtilityWidgetBlueprint.cpp
@@ -102,23 +102,8 @@ TSharedRef<SWidget> UEditorUtilityWidgetBlueprint::CreateUtilityWidget()
 		if (CreatedUMGWidget)
 		{
 			// Editor Utility is flagged as transient to prevent from dirty the World it's created in when a property added to the Utility Widget is changed
-			CreatedUMGWidget->SetFlags(RF_Transient);
-
-			// Also mark nested utility widgets as transient to prevent them from dirtying the world (since they'll be created via CreateWidget and not CreateUtilityWidget)
-			TArray<UWidget*> AllWidgets;
-			CreatedUMGWidget->WidgetTree->GetAllWidgets(AllWidgets);
-
-			for (UWidget* Widget : AllWidgets)
-			{
-				if (Widget->IsA(UEditorUtilityWidget::StaticClass()))
-				{
-					Widget->SetFlags(RF_Transient);
-					if (UPanelSlot* Slot = Widget->Slot)
-					{
-						Slot->SetFlags(RF_Transient);
-					}
-				}
-			}
+			// Also need to recursively mark nested utility widgets as transient to prevent them from dirtying the world (since they'll be created via CreateWidget and not CreateUtilityWidget)
+			MarkTransientRecursive(CreatedUMGWidget);
 		}
 	}
 
@@ -196,6 +181,26 @@ void UEditorUtilityWidgetBlueprint::ChangeTabWorld(UWorld* World, EMapChangeType
 	}
 }
 
+void UEditorUtilityWidgetBlueprint::MarkTransientRecursive(UEditorUtilityWidget* UtilityWidget)
+{
+	UtilityWidget->SetFlags(RF_Transient);
+
+	TArray<UWidget*> AllWidgets;
+	UtilityWidget->WidgetTree->GetAllWidgets(AllWidgets);
+
+	for (UWidget* Widget : AllWidgets)
+	{
+		if (UEditorUtilityWidget* ChildUtilityWidget = Cast<UEditorUtilityWidget>(Widget))
+		{
+			MarkTransientRecursive(ChildUtilityWidget);
+			if (UPanelSlot* Slot = Widget->Slot)
+			{
+				Slot->SetFlags(RF_Transient);
+			}
+		}
+	}
+}
+
 void UEditorUtilityWidgetBlueprint::UpdateRespawnListIfNeeded(TSharedRef<SDockTab> TabBeingClosed)
 {
 	if (GeneratedClass)
