Removed the warning about an out of date blueprint editor utility when the user don't have the permission to view the asset. An user cannot act upon that issue so there is no need to warn them.

diff --git a/Engine/Source/Editor/Blutility/Private/BlutilityContentBrowserExtensions.cpp b/Engine/Source/Editor/Blutility/Private/BlutilityContentBrowserExtensions.cpp
--- a/Engine/Source/Editor/Blutility/Private/BlutilityContentBrowserExtensions.cpp
+++ b/Engine/Source/Editor/Blutility/Private/BlutilityContentBrowserExtensions.cpp
@@ -2,6 +2,7 @@
 
 #include "BlutilityContentBrowserExtensions.h"
 
+#include "Algo/AnyOf.h"
 #include "AssetActionUtility.h"
 #include "AssetRegistry/AssetData.h"
 #include "BlutilityMenuExtensions.h"
@@ -10,15 +11,17 @@
 #include "Containers/Set.h"
 #include "ContentBrowserDelegates.h"
 #include "ContentBrowserModule.h"
-#include "EditorUtilityAssetPrototype.h"
 #include "Delegates/Delegate.h"
+#include "EditorUtilityAssetPrototype.h"
 #include "EditorUtilityBlueprint.h"
-#include "Algo/AnyOf.h"
 #include "Engine/Blueprint.h"
 #include "Engine/BlueprintGeneratedClass.h"
 #include "Framework/MultiBox/MultiBoxExtender.h"
 #include "HAL/Platform.h"
 #include "HAL/PlatformCrt.h"
+#include "IAssetTools.h"
+#include "Logging/MessageLog.h"
+#include "Misc/NamePermissionList.h"
 #include "Modules/ModuleManager.h"
 #include "Templates/Casts.h"
 #include "Templates/ChooseClass.h"
@@ -28,7 +31,6 @@
 #include "UObject/Class.h"
 #include "UObject/NameTypes.h"
 #include "UObject/UObjectIterator.h"
-#include "Logging/MessageLog.h"
 
 #define LOCTEXT_NAMESPACE "BlutilityContentBrowserExtensions"
 
@@ -90,10 +92,14 @@ public:
 				}
 				else
 				{
-					EditorErrors.NewPage(LOCTEXT("ScriptedActions", "Scripted Actions"));
-					TSharedRef<FTokenizedMessage> ErrorMessage = EditorErrors.Error();
-                    ErrorMessage->AddToken(FAssetNameToken::Create(ActionUtilityPrototype->GetUtilityBlueprintAsset().GetObjectPathString(),FText::FromString(ActionUtilityPrototype->GetUtilityBlueprintAsset().GetObjectPathString())));
-                    ErrorMessage->AddToken(FTextToken::Create(LOCTEXT("NeedsToBeUpdated", "needs to be re-saved and possibly upgraded.")));
+					const FAssetData& BlutilityAssetData =  ActionUtilityPrototype->GetUtilityBlueprintAsset();
+					if (IAssetTools::Get().GetAssetClassPathPermissionList(EAssetClassAction::ViewAsset)->PassesFilter(BlutilityAssetData.AssetClassPath.ToString()))
+					{
+						EditorErrors.NewPage(LOCTEXT("ScriptedActions", "Scripted Actions"));
+						TSharedRef<FTokenizedMessage> ErrorMessage = EditorErrors.Error();
+						ErrorMessage->AddToken(FAssetNameToken::Create(BlutilityAssetData.GetObjectPathString(), FText::FromString(BlutilityAssetData.GetObjectPathString())));
+						ErrorMessage->AddToken(FTextToken::Create(LOCTEXT("NeedsToBeUpdated", "needs to be re-saved and possibly upgraded.")));
+					}
 				}
 			};
 
diff --git a/Engine/Source/Editor/Blutility/Private/BlutilityLevelEditorExtensions.cpp b/Engine/Source/Editor/Blutility/Private/BlutilityLevelEditorExtensions.cpp
--- a/Engine/Source/Editor/Blutility/Private/BlutilityLevelEditorExtensions.cpp
+++ b/Engine/Source/Editor/Blutility/Private/BlutilityLevelEditorExtensions.cpp
@@ -3,21 +3,23 @@
 #include "BlutilityLevelEditorExtensions.h"
 
 #include "ActorActionUtility.h"
+#include "Algo/AnyOf.h"
 #include "AssetRegistry/AssetData.h"
 #include "BlutilityMenuExtensions.h"
 #include "Containers/Array.h"
 #include "Containers/Map.h"
 #include "Containers/Set.h"
 #include "CoreTypes.h"
-#include "EditorUtilityAssetPrototype.h"
 #include "Delegates/Delegate.h"
+#include "EditorUtilityAssetPrototype.h"
 #include "EditorUtilityBlueprint.h"
 #include "Framework/MultiBox/MultiBoxExtender.h"
 #include "GameFramework/Actor.h"
 #include "HAL/PlatformCrt.h"
+#include "IAssetTools.h"
 #include "LevelEditor.h"
-#include "Algo/AnyOf.h"
 #include "Logging/MessageLog.h"
+#include "Misc/NamePermissionList.h"
 #include "Modules/ModuleManager.h"
 #include "Templates/Casts.h"
 #include "Templates/SharedPointer.h"
@@ -73,10 +75,14 @@ public:
 				}
 				else
                 {
-					EditorErrors.NewPage(LOCTEXT("ScriptedActions", "Scripted Actions"));
-					TSharedRef<FTokenizedMessage> ErrorMessage = EditorErrors.Error();
-                    ErrorMessage->AddToken(FAssetNameToken::Create(ActionUtilityPrototype->GetUtilityBlueprintAsset().GetObjectPathString(),FText::FromString(ActionUtilityPrototype->GetUtilityBlueprintAsset().GetObjectPathString())));
-                    ErrorMessage->AddToken(FTextToken::Create(LOCTEXT("NeedsToBeUpdated", "needs to be re-saved and possibly upgraded.")));
+					const FAssetData& BlutilityAssetData =  ActionUtilityPrototype->GetUtilityBlueprintAsset();
+					if (IAssetTools::Get().GetAssetClassPathPermissionList(EAssetClassAction::ViewAsset)->PassesFilter(BlutilityAssetData.AssetClassPath.ToString()))
+					{
+						EditorErrors.NewPage(LOCTEXT("ScriptedActions", "Scripted Actions"));
+						TSharedRef<FTokenizedMessage> ErrorMessage = EditorErrors.Error();
+						ErrorMessage->AddToken(FAssetNameToken::Create(BlutilityAssetData.GetObjectPathString(), FText::FromString(BlutilityAssetData.GetObjectPathString())));
+						ErrorMessage->AddToken(FTextToken::Create(LOCTEXT("NeedsToBeUpdated", "needs to be re-saved and possibly upgraded.")));
+					}
                 }
 			};
 
