UEditorUtilityWidgetBlueprint can be mapped exactly to UEditorUtilityWidget, I was not aware that it had a more precise root class than UUserWidget

diff --git a/Engine/Source/Editor/Blutility/Private/BlutilityModule.cpp b/Engine/Source/Editor/Blutility/Private/BlutilityModule.cpp
--- a/Engine/Source/Editor/Blutility/Private/BlutilityModule.cpp
+++ b/Engine/Source/Editor/Blutility/Private/BlutilityModule.cpp
@@ -97,7 +97,7 @@ public:
 		KismetCompilerModule.OverrideBPTypeForClass(AEditorUtilityActor::StaticClass(), UEditorUtilityBlueprint::StaticClass());
 		KismetCompilerModule.OverrideBPTypeForClass(AEditorUtilityCamera::StaticClass(), UEditorUtilityBlueprint::StaticClass());
 		KismetCompilerModule.OverrideBPTypeForClass(UEditorUtilityObject::StaticClass(), UEditorUtilityBlueprint::StaticClass());
-		KismetCompilerModule.OverrideBPTypeForClassInEditor(UUserWidget::StaticClass(), UEditorUtilityWidgetBlueprint::StaticClass());
+		KismetCompilerModule.OverrideBPTypeForClass(UEditorUtilityWidget::StaticClass(), UEditorUtilityWidgetBlueprint::StaticClass());
 
 		FBlutilityContentBrowserExtensions::InstallHooks();
 		FBlutilityLevelEditorExtensions::InstallHooks();
diff --git a/Engine/Source/Editor/KismetCompiler/Public/KismetCompilerModule.h b/Engine/Source/Editor/KismetCompiler/Public/KismetCompilerModule.h
--- a/Engine/Source/Editor/KismetCompiler/Public/KismetCompilerModule.h
+++ b/Engine/Source/Editor/KismetCompiler/Public/KismetCompilerModule.h
@@ -97,6 +97,7 @@ public:
 
 	/** Facilities for establishing mappings between UClasses, UBlueprints, and UBlueprintGenerated Classes*/
 	virtual void OverrideBPTypeForClass(UClass* Class, TSubclassOf<UBlueprint> BlueprintType) = 0;
+	UE_DEPRECATED(5.4, "Conditional overrides in editor have been deprecated - make a new sentinel type if required to keep UBlueprint mappings unambiguous")
 	virtual void OverrideBPTypeForClassInEditor(UClass* Class, TSubclassOf<UBlueprint> BlueprintType) = 0;
 	virtual void OverrideBPGCTypeForBPType(TSubclassOf<UBlueprint> BlueprintType, TSubclassOf<UBlueprintGeneratedClass> BPGCType) = 0;
 
