Hide asset action menu options when the filters do not pass

diff --git a/Engine/Source/Editor/Blutility/Classes/AssetActionUtility.h b/Engine/Source/Editor/Blutility/Classes/AssetActionUtility.h
--- a/Engine/Source/Editor/Blutility/Classes/AssetActionUtility.h
+++ b/Engine/Source/Editor/Blutility/Classes/AssetActionUtility.h
@@ -35,11 +35,18 @@ struct FAssetActionSupportCondition
 
 	/**
 	 * This is the failure reason to reply to the user with if the condition above fails.
-	 * If you leave this blank, we will hide the option.  If fill in the reason, will show the option but give the reason
-	 * it is disabled.
+	 * If you fill in the reason, it will override the default failure text in the tooltip for the function menu option.
 	 */
 	UPROPERTY(EditAnywhere, Category=Condition, meta=(MultiLine=true))
 	FString FailureReason;
+
+	/**
+	 * If this filter does not pass, show the corresponding functions from the menu anyways.
+	 * If true, the menu option will display with the error tooltip, but be disabled.
+	 * If false, the menu options will be removed entirely.
+	 */
+	UPROPERTY(EditDefaultsOnly, Category = "Asset Support")
+	bool bShowInMenuIfFilterFails = true;
 };
 
 /** 
diff --git a/Engine/Source/Editor/Blutility/Private/BlutilityMenuExtensions.cpp b/Engine/Source/Editor/Blutility/Private/BlutilityMenuExtensions.cpp
--- a/Engine/Source/Editor/Blutility/Private/BlutilityMenuExtensions.cpp
+++ b/Engine/Source/Editor/Blutility/Private/BlutilityMenuExtensions.cpp
@@ -690,15 +690,16 @@ void FBlutilityMenuExtensions::CreateBlutilityActionsMenu(FMenuBuilder& MenuBuil
 	TMap<FString, TArray<FFunctionAndUtil>> CategoryFunctions;
 	ExtractFunctions(Utils, CategoryFunctions);
 	
-	auto AddFunctionEntries = [Selection, IsValidPropertyType](FMenuBuilder& SubMenuBuilder, const TArray<FFunctionAndUtil>& FunctionUtils)
+	auto AddFunctionEntries = [Selection, IsValidPropertyType](const TArray<FFunctionAndUtil>& FunctionUtils)
 	{
 		BlutilityUtil::FTextFilterExpressionContext TextFilterContext;
 		
 		FTextFilterExpressionEvaluator TextFilterExpressionEvaluator(ETextFilterExpressionEvaluatorMode::Complex);
-		
+		TArray<FMenuEntryParams> GeneratedMenuEntryParams;
 		for (const FFunctionAndUtil& FunctionAndUtil : FunctionUtils)
 		{
 			bool PassesFilterCondition = true;
+			bool bShowInMenu = true;
 			FString FilterFailureMessage;
 			
 			if constexpr ( std::is_same_v<FAssetData, SelectionType> )
@@ -723,10 +724,16 @@ void FBlutilityMenuExtensions::CreateBlutilityActionsMenu(FMenuBuilder& MenuBuil
 
 					if (!PassesFilterCondition)
 					{
+						bShowInMenu = Condition.bShowInMenuIfFilterFails;
 						break;
 					}
 				}
 			}
+
+			if (!PassesFilterCondition && !bShowInMenu)
+			{
+				continue;
+			}
 			
 			FText TooltipText;
 
@@ -739,116 +746,116 @@ void FBlutilityMenuExtensions::CreateBlutilityActionsMenu(FMenuBuilder& MenuBuil
 				TooltipText = FText::Format(LOCTEXT("AssetUtilTooltipWithErrorFormat", "{0}\n\n({1})\n\n(Shift-click to edit script)"), FunctionAndUtil.FunctionData.TooltipText, FText::FromString(FilterFailureMessage));
 			}
 
-			SubMenuBuilder.AddMenuEntry(
-				FunctionAndUtil.FunctionData.NameText,
-				TooltipText,
-				FSlateIcon(FAppStyle::GetAppStyleSetName(), "GraphEditor.Event_16x"),
-				
-				FUIAction
-				(
-					FExecuteAction::CreateLambda([FunctionAndUtil, Selection, IsValidPropertyType]
+			FMenuEntryParams MenuEntryParams;
+			MenuEntryParams.LabelOverride = FunctionAndUtil.FunctionData.NameText;
+			MenuEntryParams.ToolTipOverride = TooltipText;
+			MenuEntryParams.IconOverride = FSlateIcon(FAppStyle::GetAppStyleSetName(), "GraphEditor.Event_16x");
+			MenuEntryParams.DirectActions = FUIAction(FExecuteAction::CreateLambda([FunctionAndUtil, Selection, IsValidPropertyType]
+				{
+					if (FSlateApplication::Get().GetModifierKeys().IsShiftDown())
 					{
-						if (FSlateApplication::Get().GetModifierKeys().IsShiftDown())
-						{
-							OpenEditorForUtility(FunctionAndUtil);
-						}
-						else
-						{
-							// We dont run this on the CDO, as bad things could occur!
-							UObject* TempObject = NewObject<UObject>(GetTransientPackage(), Cast<UObject>(FunctionAndUtil.Util->LoadUtilityAsset())->GetClass());
-							TempObject->AddToRoot(); // Some Blutility actions might run GC so the TempObject needs to be rooted to avoid getting destroyed
+						OpenEditorForUtility(FunctionAndUtil);
+					}
+					else
+					{
+						// We dont run this on the CDO, as bad things could occur!
+						UObject* TempObject = NewObject<UObject>(GetTransientPackage(), Cast<UObject>(FunctionAndUtil.Util->LoadUtilityAsset())->GetClass());
+						TempObject->AddToRoot(); // Some Blutility actions might run GC so the TempObject needs to be rooted to avoid getting destroyed
 
-							UFunction* Function = FunctionAndUtil.GetFunction();
-							
-							if (Function->NumParms > 0)
-							{
-								// Create a parameter struct and fill in defaults
-								TSharedRef<FStructOnScope> FuncParams = MakeShared<FStructOnScope>(Function);
+						UFunction* Function = FunctionAndUtil.GetFunction();
+						
+						if (Function->NumParms > 0)
+						{
+							// Create a parameter struct and fill in defaults
+							TSharedRef<FStructOnScope> FuncParams = MakeShared<FStructOnScope>(Function);
 
-								FProperty* FirstParamProperty = nullptr;
+							FProperty* FirstParamProperty = nullptr;
 
-								int32 ParameterIndex = 0;
-								for (TFieldIterator<FProperty> It(Function); It&& It->HasAnyPropertyFlags(CPF_Parm); ++It)
+							int32 ParameterIndex = 0;
+							for (TFieldIterator<FProperty> It(Function); It&& It->HasAnyPropertyFlags(CPF_Parm); ++It)
+							{
+								FString Defaults;
+								if (UEdGraphSchema_K2::FindFunctionParameterDefaultValue(Function, *It, Defaults))
 								{
-									FString Defaults;
-									if (UEdGraphSchema_K2::FindFunctionParameterDefaultValue(Function, *It, Defaults))
-									{
-										It->ImportText_Direct(*Defaults, It->ContainerPtrToValuePtr<uint8>(FuncParams->GetStructMemory()), nullptr, PPF_None);
-									}
-
-									// Check to see if the first parameter matches the selection object type, in that case we can directly forward the selection to it
-									if (ParameterIndex == 0 && IsValidPropertyType(*It))
-									{
-										FirstParamProperty = *It;
-									}
-
-									++ParameterIndex;
+									It->ImportText_Direct(*Defaults, It->ContainerPtrToValuePtr<uint8>(FuncParams->GetStructMemory()), nullptr, PPF_None);
 								}
 
-								bool bApply = true;
-
-								if (!FirstParamProperty || ParameterIndex > 1)
+								// Check to see if the first parameter matches the selection object type, in that case we can directly forward the selection to it
+								if (ParameterIndex == 0 && IsValidPropertyType(*It))
 								{
-									// pop up a dialog to input params to the function
-									TSharedRef<SWindow> Window = SNew(SWindow)
-										.Title(Function->GetDisplayNameText())
-										.ClientSize(FVector2D(400, 200))
-										.SupportsMinimize(false)
-										.SupportsMaximize(false);
-
-									TSharedPtr<SFunctionParamDialog> Dialog;
-									Window->SetContent(
-										SAssignNew(Dialog, SFunctionParamDialog, Window, FuncParams, FirstParamProperty ? FirstParamProperty->GetFName() : NAME_None)
-										.OkButtonText(LOCTEXT("OKButton", "OK"))
-										.OkButtonTooltipText(Function->GetToolTipText()));
-
-									GEditor->EditorAddModalWindow(Window);
-									bApply = Dialog->bOKPressed;
+									FirstParamProperty = *It;
 								}
 
+								++ParameterIndex;
+							}
 
-								if (bApply)
+							bool bApply = true;
+
+							if (!FirstParamProperty || ParameterIndex > 1)
+							{
+								// pop up a dialog to input params to the function
+								TSharedRef<SWindow> Window = SNew(SWindow)
+									.Title(Function->GetDisplayNameText())
+									.ClientSize(FVector2D(400, 200))
+									.SupportsMinimize(false)
+									.SupportsMaximize(false);
+
+								TSharedPtr<SFunctionParamDialog> Dialog;
+								Window->SetContent(
+									SAssignNew(Dialog, SFunctionParamDialog, Window, FuncParams, FirstParamProperty ? FirstParamProperty->GetFName() : NAME_None)
+									.OkButtonText(LOCTEXT("OKButton", "OK"))
+									.OkButtonTooltipText(Function->GetToolTipText()));
+
+								GEditor->EditorAddModalWindow(Window);
+								bApply = Dialog->bOKPressed;
+							}
+
+
+							if (bApply)
+							{
+								FScopedTransaction Transaction(NSLOCTEXT("UnrealEd", "BlutilityAction", "Blutility Action"));
+								FEditorScriptExecutionGuard ScriptGuard;
+								const bool bForwardUserSelection = FirstParamProperty != nullptr;
+								if (bForwardUserSelection)
 								{
-									FScopedTransaction Transaction(NSLOCTEXT("UnrealEd", "BlutilityAction", "Blutility Action"));
-									FEditorScriptExecutionGuard ScriptGuard;
-									const bool bForwardUserSelection = FirstParamProperty != nullptr;
-									if (bForwardUserSelection)
-									{
-										// For each user-select asset forward the selection object into the function first's parameter (if it matches)
-										const FString Path = FirstParamProperty->GetPathName(Function);
-
-										// Ensure we only process selection objects that are valid for this function/utility
-										for (const int32& SelectionIndex : FunctionAndUtil.SelectionIndices)
-										{
-											const auto SelectedAsset = Selection[SelectionIndex];
-											FirstParamProperty->CopySingleValue(FirstParamProperty->ContainerPtrToValuePtr<uint8>(FuncParams->GetStructMemory()), &SelectedAsset);
-											TempObject->ProcessEvent(Function, FuncParams->GetStructMemory());
-										}
-									}
-									else
+									// For each user-select asset forward the selection object into the function first's parameter (if it matches)
+									const FString Path = FirstParamProperty->GetPathName(Function);
+
+									// Ensure we only process selection objects that are valid for this function/utility
+									for (const int32& SelectionIndex : FunctionAndUtil.SelectionIndices)
 									{
-										// User is expected to manage the asset selection on its own
+										const auto SelectedAsset = Selection[SelectionIndex];
+										FirstParamProperty->CopySingleValue(FirstParamProperty->ContainerPtrToValuePtr<uint8>(FuncParams->GetStructMemory()), &SelectedAsset);
 										TempObject->ProcessEvent(Function, FuncParams->GetStructMemory());
 									}
 								}
+								else
+								{
+									// User is expected to manage the asset selection on its own
+									TempObject->ProcessEvent(Function, FuncParams->GetStructMemory());
+								}
 							}
-							else
-							{
-								FScopedTransaction Transaction(NSLOCTEXT("UnrealEd", "BlutilityAction", "Blutility Action"));
-								FEditorScriptExecutionGuard ScriptGuard;
-								TempObject->ProcessEvent(Function, nullptr);
-							}
-
-							TempObject->RemoveFromRoot();
 						}
-					}),
-					FCanExecuteAction::CreateLambda([PassesFilterCondition]()
-					{
-						return PassesFilterCondition;
-					})
-				)
+						else
+						{
+							FScopedTransaction Transaction(NSLOCTEXT("UnrealEd", "BlutilityAction", "Blutility Action"));
+							FEditorScriptExecutionGuard ScriptGuard;
+							TempObject->ProcessEvent(Function, nullptr);
+						}
+
+						TempObject->RemoveFromRoot();
+					}
+				}),
+				FCanExecuteAction::CreateLambda([PassesFilterCondition]()
+				{
+					return PassesFilterCondition;
+				})
 			);
+
+			GeneratedMenuEntryParams.Emplace(MenuEntryParams);
 		}
+
+		return GeneratedMenuEntryParams;
 	};
 
 	// Add a menu item for each function
@@ -867,20 +874,29 @@ void FBlutilityMenuExtensions::CreateBlutilityActionsMenu(FMenuBuilder& MenuBuil
 				// Add functions belong to the same category to a sub-menu
 				for (const FString& CategoryName : CategoryNames)
 				{
-					const TArray<FFunctionAndUtil>& FunctionUtils = CategoryFunctions.FindChecked(CategoryName);
-					InMenuBuilder.AddSubMenu(FText::FromString(CategoryName), FText::FromString(CategoryName),
-						FNewMenuDelegate::CreateLambda([FunctionUtils, AddFunctionEntries](FMenuBuilder& InSubMenuBuilder) 
-						{
-							AddFunctionEntries(InSubMenuBuilder, FunctionUtils);
-						})
-					);
+					const TArray<FMenuEntryParams> GeneratedMenuEntries = AddFunctionEntries(CategoryFunctions.FindChecked(CategoryName));
+					if (GeneratedMenuEntries.Num() > 0)
+					{
+						InMenuBuilder.AddSubMenu(FText::FromString(CategoryName), FText::FromString(CategoryName),
+							FNewMenuDelegate::CreateLambda([GeneratedMenuEntries](FMenuBuilder& InSubMenuBuilder)
+								{
+									for (const FMenuEntryParams& MenuParams : GeneratedMenuEntries)
+									{
+										InSubMenuBuilder.AddMenuEntry(MenuParams);
+									}
+								})
+						);
+					}
 				}
 
 				// Non-categorized functions
 				const TArray<FFunctionAndUtil>* DefaultCategoryFunctionsPtr = CategoryFunctions.Find(FString());
 				if (DefaultCategoryFunctionsPtr)
 				{
-					AddFunctionEntries(InMenuBuilder, *DefaultCategoryFunctionsPtr);
+					for (const FMenuEntryParams& MenuParams : AddFunctionEntries(*DefaultCategoryFunctionsPtr))
+					{
+						InMenuBuilder.AddMenuEntry(MenuParams);
+					}
 				}
 			}),
 			false,
