Add GeneratedClass API to Editor Utility

diff --git a/Engine/Source/Editor/Blutility/Classes/EditorUtilityWidgetBlueprint.h b/Engine/Source/Editor/Blutility/Classes/EditorUtilityWidgetBlueprint.h
--- a/Engine/Source/Editor/Blutility/Classes/EditorUtilityWidgetBlueprint.h
+++ b/Engine/Source/Editor/Blutility/Classes/EditorUtilityWidgetBlueprint.h
@@ -77,13 +77,14 @@ public:
 	virtual UWidgetEditingProjectSettings* GetRelevantSettings() override;
 	virtual const UWidgetEditingProjectSettings* GetRelevantSettings() const override;
 
+public:
+	static void MarkTransientRecursive(UEditorUtilityWidget* UtilityWidget);
+
 private:
 	bool IsWidgetEnabled() const;
 
 	void ChangeTabWorld(UWorld* World, EMapChangeType MapChangeType);
 
-	void MarkTransientRecursive(UEditorUtilityWidget* UtilityWidget);
-
 private:
 	// Should the widget be enabled when running PIE
 	UPROPERTY(Category = Settings, EditDefaultsOnly)
diff --git a/Engine/Source/Editor/Blutility/Private/BlutilityMenuExtensions.cpp b/Engine/Source/Editor/Blutility/Private/BlutilityMenuExtensions.cpp
--- a/Engine/Source/Editor/Blutility/Private/BlutilityMenuExtensions.cpp
+++ b/Engine/Source/Editor/Blutility/Private/BlutilityMenuExtensions.cpp
@@ -2,6 +2,8 @@
 
 #include "BlutilityMenuExtensions.h"
 
+#include "ActorActionUtility.h"
+#include "AssetActionUtility.h"
 #include "AssetRegistry/ARFilter.h"
 #include "AssetRegistry/AssetData.h"
 #include "AssetRegistry/AssetDataTagMap.h"
@@ -17,8 +19,10 @@
 #include "EditorUtilityAssetPrototype.h"
 #include "Editor/EditorEngine.h"
 #include "EditorUtilityBlueprint.h"
+#include "EditorUtilityWidgetProjectSettings.h"
 #include "FrontendFilters.h"
 #include "Engine/Blueprint.h"
+#include "Engine/BlueprintGeneratedClass.h"
 #include "Framework/Application/SlateApplication.h"
 #include "Framework/Commands/UIAction.h"
 #include "Framework/MultiBox/MultiBoxBuilder.h"
@@ -601,7 +605,48 @@ void FBlutilityMenuExtensions::GetBlutilityClasses(TArray<FAssetData>& OutAssets
 			}
 		}
 	}
+
+	const UEditorUtilityWidgetProjectSettings* EditorUtilitySettings = GetDefault<UEditorUtilityWidgetProjectSettings>();
+
+	if (EditorUtilitySettings->bSearchGeneratedClassesForScriptedActions)
+	{
+		auto FilterAssets = [&AssetRegistry, &OutAssets](const FNamePermissionList& PermissionList)
+		{
+			FARFilter GeneratedFilter;
+			GeneratedFilter.ClassPaths.Add(UBlueprintGeneratedClass::StaticClass()->GetClassPathName());
+			GeneratedFilter.bRecursiveClasses = true;
+			GeneratedFilter.bRecursivePaths = true;
+
+			TArray<FAssetData> GeneratedAssetList;
+			AssetRegistry.GetAssets(GeneratedFilter, GeneratedAssetList);
+
+			for (const FAssetData& Asset : GeneratedAssetList)
+			{
+				// Abstract or deprecated blutilities should not be included.
+				const EClassFlags BPFlags = static_cast<EClassFlags>(Asset.GetTagValueRef<uint32>(FBlueprintTags::ClassFlags));
+				if (EnumHasAnyFlags(BPFlags, CLASS_Abstract | CLASS_Deprecated))
+				{
+					continue;
+				}
+
+				if (PermissionList.PassesFilter(FName(Asset.GetObjectPathString())))
+				{
+					OutAssets.Add(Asset);
+				}
+			}
+		};
+
+		if (InClassName == UAssetActionUtility::StaticClass()->GetClassPathName())
+		{
+			FilterAssets(EditorUtilitySettings->GetAllowedEditorUtilityAssetActions());
+		}
+		else if (InClassName == UActorActionUtility::StaticClass()->GetClassPathName())
+		{
+			FilterAssets(EditorUtilitySettings->GetAllowedEditorUtilityActorActions());
+		}
+	}
 }
+
 void FBlutilityMenuExtensions::CreateActorBlutilityActionsMenu(FMenuBuilder& MenuBuilder, TMap<TSharedRef<FAssetActionUtilityPrototype>, TSet<int32>> Utils, const TArray<AActor*> SelectedSupportedActors)
 {
 	CreateBlutilityActionsMenu<AActor*>(MenuBuilder, Utils,
@@ -737,7 +782,11 @@ void FBlutilityMenuExtensions::CreateBlutilityActionsMenu(FMenuBuilder& MenuBuil
 			
 			FText TooltipText;
 
-			if (FilterFailureMessage.IsEmpty())
+			if (FunctionAndUtil.Util->GetUtilityBlueprintAsset().AssetClassPath == UBlueprintGeneratedClass::StaticClass()->GetClassPathName())
+			{
+				TooltipText = LOCTEXT("AssetUtilTooltip", "Click to execute");
+			}
+			else if (FilterFailureMessage.IsEmpty())
 			{
 				TooltipText = FText::Format(LOCTEXT("AssetUtilTooltipFormat", "{0}\n\n(Shift-click to edit script)"), FunctionAndUtil.FunctionData.TooltipText);	
 			}
diff --git a/Engine/Source/Editor/Blutility/Private/EditorUtilityAssetPrototype.cpp b/Engine/Source/Editor/Blutility/Private/EditorUtilityAssetPrototype.cpp
--- a/Engine/Source/Editor/Blutility/Private/EditorUtilityAssetPrototype.cpp
+++ b/Engine/Source/Editor/Blutility/Private/EditorUtilityAssetPrototype.cpp
@@ -4,6 +4,7 @@
 
 #include "AssetActionUtility.h"
 #include "EditorUtilityBlueprint.h"
+#include "Engine/BlueprintGeneratedClass.h"
 #include "JsonObjectConverter.h"
 #include "AssetRegistry/AssetData.h"
 #include "Serialization/JsonSerializerMacros.h"
@@ -25,6 +26,10 @@ UObject* FAssetActionUtilityPrototype::LoadUtilityAsset() const
 			return BPClass->GetDefaultObject();
 		}
 	}
+	else if (const UBlueprintGeneratedClass* BPClass = Cast<UBlueprintGeneratedClass>(UtilityBlueprintAsset.GetAsset()))
+	{
+		return BPClass->GetDefaultObject();
+	}
 
 	return nullptr;
 };
diff --git a/Engine/Source/Editor/Blutility/Private/EditorUtilitySubsystem.cpp b/Engine/Source/Editor/Blutility/Private/EditorUtilitySubsystem.cpp
--- a/Engine/Source/Editor/Blutility/Private/EditorUtilitySubsystem.cpp
+++ b/Engine/Source/Editor/Blutility/Private/EditorUtilitySubsystem.cpp
@@ -84,6 +84,9 @@ void UEditorUtilitySubsystem::Initialize(FSubsystemCollectionBase& Collection)
 
 	FEditorDelegates::BeginPIE.AddUObject(this, &UEditorUtilitySubsystem::HandleOnBeginPIE);
 	FEditorDelegates::EndPIE.AddUObject(this, &UEditorUtilitySubsystem::HandleOnEndPIE);
+
+	FLevelEditorModule& LevelEditor = FModuleManager::LoadModuleChecked<FLevelEditorModule>("LevelEditor");
+	LevelEditor.OnMapChanged().AddUObject(this, &UEditorUtilitySubsystem::OnMapChanged);
 }
 
 void UEditorUtilitySubsystem::Deinitialize()
@@ -99,6 +102,11 @@ void UEditorUtilitySubsystem::Deinitialize()
 
 	FEditorDelegates::BeginPIE.RemoveAll(this);
 	FEditorDelegates::EndPIE.RemoveAll(this);
+
+	if (FLevelEditorModule* LevelEditor = FModuleManager::GetModulePtr<FLevelEditorModule>("LevelEditor"))
+	{
+		LevelEditor->OnMapChanged().RemoveAll(this);
+	}
 }
 
 void UEditorUtilitySubsystem::AddReferencedObjects(UObject* InThis, FReferenceCollector& Collector)
@@ -158,7 +166,7 @@ bool UEditorUtilitySubsystem::TryRun(UObject* Asset)
 
 	static const FName RunFunctionName("Run");
 	UFunction* RunFunction = ObjectClass->FindFunctionByName(RunFunctionName);
-	if (RunFunction)
+	if (RunFunction && RunFunction->ParmsSize == 0)
 	{
 		UObject* Instance = NewObject<UObject>(this, ObjectClass);
 		ObjectInstances.Add(Asset, Instance);
@@ -169,7 +177,28 @@ bool UEditorUtilitySubsystem::TryRun(UObject* Asset)
 	}
 	else
 	{
-		UE_LOG(LogEditorUtilityBlueprint, Warning, TEXT("Missing function named 'Run': %s"), *Asset->GetPathName());
+		UE_LOG(LogEditorUtilityBlueprint, Warning, TEXT("Missing 0 param function named 'Run': %s"), *Asset->GetPathName());
+	}
+
+	return false;
+}
+
+bool UEditorUtilitySubsystem::TryRunClass(UClass* ObjectClass)
+{
+	static const FName RunFunctionName("Run");
+	UFunction* RunFunction = ObjectClass->FindFunctionByName(RunFunctionName);
+	if (RunFunction && RunFunction->ParmsSize == 0)
+	{
+		UObject* Instance = NewObject<UObject>(this, ObjectClass);
+		ObjectInstances.Add(ObjectClass, Instance);
+
+		FEditorScriptExecutionGuard ScriptGuard;
+		Instance->ProcessEvent(RunFunction, nullptr);
+		return true;
+	}
+	else
+	{
+		UE_LOG(LogEditorUtilityBlueprint, Warning, TEXT("Missing 0 param function named 'Run': %s"), *ObjectClass->GetPathName());
 	}
 
 	return false;
@@ -218,7 +247,7 @@ UEditorUtilityWidget* UEditorUtilitySubsystem::SpawnAndRegisterTab(class UEditor
 
 UEditorUtilityWidget* UEditorUtilitySubsystem::SpawnAndRegisterTabWithId(class UEditorUtilityWidgetBlueprint* InBlueprint, FName InTabID)
 {
-	return SpawnAndRegisterTabAndGetID(InBlueprint, InTabID);;
+	return SpawnAndRegisterTabAndGetID(InBlueprint, InTabID);
 }
 
 void UEditorUtilitySubsystem::RegisterTabAndGetID(class UEditorUtilityWidgetBlueprint* InBlueprint, FName& NewTabID)
@@ -244,6 +273,72 @@ void UEditorUtilitySubsystem::RegisterTabAndGetID(class UEditorUtilityWidgetBlue
 	}
 }
 
+UEditorUtilityWidget* UEditorUtilitySubsystem::SpawnAndRegisterTabAndGetIDGeneratedClass(UWidgetBlueprintGeneratedClass* InGeneratedWidgetBlueprint, FName& NewTabID)
+{
+	FName TabID;
+	RegisterTabAndGetIDGeneratedClass(InGeneratedWidgetBlueprint, TabID);
+	SpawnRegisteredTabByID(TabID);
+	NewTabID = TabID;
+
+	FLevelEditorModule& LevelEditorModule = FModuleManager::GetModuleChecked<FLevelEditorModule>(TEXT("LevelEditor"));
+	if (TSharedPtr<FTabManager> LevelEditorTabManager = LevelEditorModule.GetLevelEditorTabManager())
+	{
+		if (TSharedPtr<SDockTab> FoundTab = LevelEditorTabManager->FindExistingLiveTab(NewTabID))
+		{
+			if (UEditorUtilityWidget **CreatedUMGWidget = SpawnedFromGeneratedClassTabs.Find(FoundTab.ToSharedRef()))
+			{
+				return *CreatedUMGWidget;
+			}
+		}
+	}
+	return nullptr;
+}
+
+UEditorUtilityWidget* UEditorUtilitySubsystem::SpawnAndRegisterTabGeneratedClass(UWidgetBlueprintGeneratedClass* InGeneratedWidgetBlueprint)
+{
+	FName InTabID;
+	return SpawnAndRegisterTabAndGetIDGeneratedClass(InGeneratedWidgetBlueprint, InTabID);
+}
+
+UEditorUtilityWidget* UEditorUtilitySubsystem::SpawnAndRegisterTabWithIdGeneratedClass(UWidgetBlueprintGeneratedClass* InGeneratedWidgetBlueprint, FName InTabID)
+{
+	return SpawnAndRegisterTabAndGetIDGeneratedClass(InGeneratedWidgetBlueprint, InTabID);
+}
+
+void UEditorUtilitySubsystem::RegisterTabAndGetIDGeneratedClass(UWidgetBlueprintGeneratedClass* InGeneratedWidgetBlueprint, FName& NewTabID)
+{
+	if (InGeneratedWidgetBlueprint && !IsRunningCommandlet())
+	{
+		FName RegistrationName = NewTabID.IsNone() ? FName(*(InGeneratedWidgetBlueprint->GetPathName() + LOCTEXT("ActiveTabSuffix", "_ActiveTab").ToString())) : FName(*(InGeneratedWidgetBlueprint->GetPathName() + NewTabID.ToString()));
+		FText DisplayName;
+		const UEditorUtilityWidget* EditorUtilityWidget = InGeneratedWidgetBlueprint->GetDefaultObject<UEditorUtilityWidget>();
+
+		if (EditorUtilityWidget && !EditorUtilityWidget->GetTabDisplayName().IsEmpty())
+		{
+			DisplayName = EditorUtilityWidget->GetTabDisplayName();
+		}
+		else
+		{
+			DisplayName = FText::FromString(FName::NameToDisplayString(InGeneratedWidgetBlueprint->GetName(), false));
+		}
+
+		FLevelEditorModule& LevelEditorModule = FModuleManager::GetModuleChecked<FLevelEditorModule>(TEXT("LevelEditor"));
+		if (TSharedPtr<FTabManager> LevelEditorTabManager = LevelEditorModule.GetLevelEditorTabManager())
+		{
+			if (!LevelEditorTabManager->HasTabSpawner(RegistrationName))
+			{
+				IBlutilityModule* BlutilityModule = FModuleManager::GetModulePtr<IBlutilityModule>("Blutility");
+				LevelEditorTabManager->RegisterTabSpawner(RegistrationName, FOnSpawnTab::CreateUObject(this, &UEditorUtilitySubsystem::SpawnEditorUITabFromGeneratedClass, InGeneratedWidgetBlueprint))
+					.SetDisplayName(DisplayName)
+					.SetGroup(BlutilityModule->GetMenuGroup().ToSharedRef());
+			}
+
+			RegisteredTabsByGeneratedClass.Add(RegistrationName, InGeneratedWidgetBlueprint);
+			NewTabID = RegistrationName;
+		}
+	}
+}
+
 bool UEditorUtilitySubsystem::SpawnRegisteredTabByID(FName NewTabID)
 {
 	if (!IsRunningCommandlet())
@@ -570,4 +665,78 @@ void UEditorUtilitySubsystem::HandleOnEndPIE(const bool bIsSimulating)
 	OnEndPIE.Broadcast(bIsSimulating);
 }
 
+TSharedRef<SDockTab> UEditorUtilitySubsystem::SpawnEditorUITabFromGeneratedClass(const FSpawnTabArgs& SpawnTabArgs, UWidgetBlueprintGeneratedClass* InGeneratedWidgetBlueprint)
+{
+	TSharedRef<SDockTab> SpawnedTab = SNew(SDockTab);
+	UEditorUtilityWidget* CreatedUMGWidget = nullptr;
+
+	auto CreateUtilityWidgetFromGeneratedClass = [&CreatedUMGWidget](UWidgetBlueprintGeneratedClass* InGeneratedWidgetBlueprint)
+	{
+		TSharedRef<SWidget> TabWidget = SNullWidget::NullWidget;
+
+		UClass* BlueprintClass = InGeneratedWidgetBlueprint;
+		TSubclassOf<UEditorUtilityWidget> WidgetClass = BlueprintClass;
+
+		if (UWorld* World = GEditor->GetEditorWorldContext().World())
+		{
+			CreatedUMGWidget = CreateWidget<UEditorUtilityWidget>(World, WidgetClass);
+			if (CreatedUMGWidget)
+			{
+				// Editor Utility is flagged as transient to prevent from dirty the World it's created in when a property added to the Utility Widget is changed
+				// Also need to recursively mark nested utility widgets as transient to prevent them from dirtying the world (since they'll be created via CreateWidget and not CreateUtilityWidget)
+				UEditorUtilityWidgetBlueprint::MarkTransientRecursive(CreatedUMGWidget);
+			}
+		}
+
+		if (CreatedUMGWidget)
+		{
+			TabWidget = SNew(SVerticalBox)
+				+ SVerticalBox::Slot()
+				.HAlign(HAlign_Fill)
+				[
+					CreatedUMGWidget->TakeWidget()
+				];
+		}
+		return TabWidget;
+	};
+
+	TSharedRef<SWidget> TabWidget = CreateUtilityWidgetFromGeneratedClass(InGeneratedWidgetBlueprint);
+	SpawnedTab->SetContent(TabWidget);
+
+	SpawnedTab->SetOnTabClosed(SDockTab::FOnTabClosedCallback::CreateUObject(this, &UEditorUtilitySubsystem::OnSpawnedFromGeneratedClassTabClosed));
+	SpawnedFromGeneratedClassTabs.Add(SpawnedTab, CreatedUMGWidget);
+
+	return SpawnedTab;
+}
+
+void UEditorUtilitySubsystem::OnSpawnedFromGeneratedClassTabClosed(TSharedRef<SDockTab> TabBeingClosed)
+{
+	if (UEditorUtilityWidget** Widget = SpawnedFromGeneratedClassTabs.Find(TabBeingClosed))
+	{
+		(*Widget)->Rename(nullptr, GetTransientPackage());
+	}
+
+	SpawnedFromGeneratedClassTabs.Remove(TabBeingClosed);
+}
+
+void UEditorUtilitySubsystem::OnMapChanged(UWorld* World, EMapChangeType MapChangeType)
+{
+	if (MapChangeType != EMapChangeType::SaveMap)
+	{
+		// We need to Delete the UMG widget if we are tearing down the World it was built with.
+		for (TPair<TSharedRef<SDockTab>, UEditorUtilityWidget*> SpawnedFromGeneratedClassTab : SpawnedFromGeneratedClassTabs)
+		{
+			TSharedRef<SDockTab> CreatedTab = SpawnedFromGeneratedClassTab.Key;
+			UEditorUtilityWidget* CreatedUMGWidget = SpawnedFromGeneratedClassTab.Value;
+			if (CreatedUMGWidget && World == CreatedUMGWidget->GetWorld())
+			{
+				CreatedTab->SetContent(SNullWidget::NullWidget);
+
+				CreatedUMGWidget->Rename(nullptr, GetTransientPackage());
+			}
+		}
+		SpawnedFromGeneratedClassTabs.Empty();
+	}
+}
+
 #undef LOCTEXT_NAMESPACE
diff --git a/Engine/Source/Editor/Blutility/Private/EditorUtilityWidgetProjectSettings.cpp b/Engine/Source/Editor/Blutility/Private/EditorUtilityWidgetProjectSettings.cpp
--- a/Engine/Source/Editor/Blutility/Private/EditorUtilityWidgetProjectSettings.cpp
+++ b/Engine/Source/Editor/Blutility/Private/EditorUtilityWidgetProjectSettings.cpp
@@ -42,3 +42,23 @@ void UEditorUtilityWidgetProjectSettings::PostEditChangeChainProperty(FPropertyC
 		}
 	}
 }
+
+FNamePermissionList& UEditorUtilityWidgetProjectSettings::GetAllowedEditorUtilityActorActions()
+{
+	return AllowedEditorUtilityActorActions;
+}
+
+const FNamePermissionList& UEditorUtilityWidgetProjectSettings::GetAllowedEditorUtilityActorActions() const
+{
+	return AllowedEditorUtilityActorActions;
+}
+
+FNamePermissionList& UEditorUtilityWidgetProjectSettings::GetAllowedEditorUtilityAssetActions()
+{
+	return AllowedEditorUtilityAssetActions;
+}
+
+const FNamePermissionList& UEditorUtilityWidgetProjectSettings::GetAllowedEditorUtilityAssetActions() const
+{
+	return AllowedEditorUtilityAssetActions;
+}
diff --git a/Engine/Source/Editor/Blutility/Public/EditorUtilitySubsystem.h b/Engine/Source/Editor/Blutility/Public/EditorUtilitySubsystem.h
--- a/Engine/Source/Editor/Blutility/Public/EditorUtilitySubsystem.h
+++ b/Engine/Source/Editor/Blutility/Public/EditorUtilitySubsystem.h
@@ -17,17 +17,21 @@
 #include "UObject/ObjectMacros.h"
 #include "UObject/ObjectPtr.h"
 #include "UObject/UObjectGlobals.h"
+#include "UnrealEdMisc.h"
 
 #include "EditorUtilitySubsystem.generated.h"
 
 class FOutputDevice;
+class FSpawnTabArgs;
 class FSubsystemCollectionBase;
 class IConsoleObject;
+class SDockTab;
 class SWindow;
 class UClass;
 class UEditorUtilityTask;
 class UEditorUtilityWidget;
 class UEditorUtilityWidgetBlueprint;
+class UWidgetBlueprintGeneratedClass;
 class UWorld;
 struct FFrame;
 struct FSoftObjectPath;
@@ -58,6 +62,7 @@ public:
 	TArray<FSoftObjectPath> StartupObjects;
 
 	TMap<FName, UEditorUtilityWidgetBlueprint*> RegisteredTabs;
+	TMap<FName, UWidgetBlueprintGeneratedClass*> RegisteredTabsByGeneratedClass;
 
 	// Allow startup object to be garbage collected
 	UFUNCTION(BlueprintCallable, Category = "Development|Editor")
@@ -66,6 +71,9 @@ public:
 	UFUNCTION(BlueprintCallable, Category = "Development|Editor")
 	bool TryRun(UObject* Asset);
 
+	UFUNCTION(BlueprintCallable, Category = "Development|Editor")
+	bool TryRunClass(UClass* ObjectClass);
+
 	UFUNCTION(BlueprintCallable, Category = "Development|Editor")
 	bool CanRun(UObject* Asset) const;
 
@@ -74,6 +82,9 @@ public:
 
 	UFUNCTION(BlueprintCallable, Category = "Development|Editor")
 	UEditorUtilityWidget* SpawnAndRegisterTab(class UEditorUtilityWidgetBlueprint* InBlueprint);
+	
+	UFUNCTION(BlueprintCallable, Category = "Development|Editor")
+	void RegisterTabAndGetID(class UEditorUtilityWidgetBlueprint* InBlueprint, FName& NewTabID);
 
 	/**
 	 * Unlike SpawnAndRegisterTabAndGetID allows spawn tab while providing TabID from Python scripts or BP
@@ -82,7 +93,19 @@ public:
 	UEditorUtilityWidget* SpawnAndRegisterTabWithId(class UEditorUtilityWidgetBlueprint* InBlueprint, FName InTabID);
 
 	UFUNCTION(BlueprintCallable, Category = "Development|Editor")
-	void RegisterTabAndGetID(class UEditorUtilityWidgetBlueprint* InBlueprint, FName& NewTabID);
+	UEditorUtilityWidget* SpawnAndRegisterTabAndGetIDGeneratedClass(UWidgetBlueprintGeneratedClass* InGeneratedWidgetBlueprint, FName& NewTabID);
+
+	UFUNCTION(BlueprintCallable, Category = "Development|Editor")
+	UEditorUtilityWidget* SpawnAndRegisterTabGeneratedClass(UWidgetBlueprintGeneratedClass* InGeneratedWidgetBlueprint);
+
+	UFUNCTION(BlueprintCallable, Category = "Development|Editor")
+	void RegisterTabAndGetIDGeneratedClass(UWidgetBlueprintGeneratedClass* InGeneratedWidgetBlueprint, FName& NewTabID);
+
+	/**
+	 * Unlike SpawnAndRegisterTabAndGetID allows spawn tab while providing TabID from Python scripts or BP
+	 */
+	UFUNCTION(BlueprintCallable, Category = "Development|Editor")
+	UEditorUtilityWidget* SpawnAndRegisterTabWithIdGeneratedClass(UWidgetBlueprintGeneratedClass* InGeneratedWidgetBlueprint, FName InTabID);
 
 	/** Given an ID for a tab, try to find a tab spawner that matches, and then spawn a tab. Returns true if it was able to find a matching tab spawner */
 	UFUNCTION(BlueprintCallable, Category = "Development|Editor")
@@ -142,6 +165,15 @@ protected:
 	/** Called when Play in Editor stops. */
 	void HandleOnEndPIE(const bool bIsSimulating);
 
+	TSharedRef<SDockTab> SpawnEditorUITabFromGeneratedClass(const FSpawnTabArgs& SpawnTabArgs, UWidgetBlueprintGeneratedClass* InGeneratedWidgetBlueprint);
+
+	void OnSpawnedFromGeneratedClassTabClosed(TSharedRef<SDockTab> TabBeingClosed);
+
+	void OnMapChanged(UWorld* World, EMapChangeType MapChangeType);
+
+	TMap<TSharedRef<SDockTab>, UEditorUtilityWidget*> SpawnedFromGeneratedClassTabs;
+
+
 private:
 	IConsoleObject* RunTaskCommandObject = nullptr;
 	IConsoleObject* CancelAllTasksCommandObject = nullptr;
diff --git a/Engine/Source/Editor/Blutility/Public/EditorUtilityWidgetProjectSettings.h b/Engine/Source/Editor/Blutility/Public/EditorUtilityWidgetProjectSettings.h
--- a/Engine/Source/Editor/Blutility/Public/EditorUtilityWidgetProjectSettings.h
+++ b/Engine/Source/Editor/Blutility/Public/EditorUtilityWidgetProjectSettings.h
@@ -34,4 +34,23 @@ public:
 	virtual void PostEditChangeChainProperty(struct FPropertyChangedChainEvent& PropertyChangedEvent) override;
 	// End UObject Interface
 
+	/** Get the permission list that controls actor actions */
+	FNamePermissionList& GetAllowedEditorUtilityActorActions();
+	const FNamePermissionList& GetAllowedEditorUtilityActorActions() const;
+
+	/** Get the permission list that controls asset actions */
+	FNamePermissionList& GetAllowedEditorUtilityAssetActions();
+	const FNamePermissionList& GetAllowedEditorUtilityAssetActions() const;
+
+	/** Set true to search generated classes for scripted actions. */
+	UPROPERTY(EditAnywhere, config, Category = Designer)
+	bool bSearchGeneratedClassesForScriptedActions = false;
+
+protected:
+
+	/** EditorUtility actor actions to always consider during context menu creation */
+	FNamePermissionList AllowedEditorUtilityActorActions;
+
+	/** EditorUtility asset actions to always consider during context menu creation */
+	FNamePermissionList AllowedEditorUtilityAssetActions;
 };
