Added bIsActionForBlueprints to UAssetActionUtility for better UX than needing to override a function

diff --git a/Engine/Source/Editor/Blutility/Classes/AssetActionUtility.h b/Engine/Source/Editor/Blutility/Classes/AssetActionUtility.h
--- a/Engine/Source/Editor/Blutility/Classes/AssetActionUtility.h
+++ b/Engine/Source/Editor/Blutility/Classes/AssetActionUtility.h
@@ -70,9 +70,10 @@ public:
 
 	/**
 	 * Returns whether or not this action is designed to work specifically on Blueprints (true) or on all assets (false).
-	 * If true, GetSupportedClass() is treated as a filter against the Parent Class of selected Blueprint assets
+	 * If true, GetSupportedClass() is treated as a filter against the Parent Class of selected Blueprint assets.
+	 * @note Returns the value of bIsActionForBlueprints by default.
 	 */
-	UFUNCTION(BlueprintPure, BlueprintImplementableEvent, Category="Assets")
+	UFUNCTION(BlueprintPure, BlueprintNativeEvent, Category="Assets")
 	bool IsActionForBlueprints() const;
 
 	/**
@@ -88,9 +89,17 @@ public:
 	// End UObject
 
 protected:
+	/**
+	 * Is this action designed to work specifically on Blueprints (true) or on all assets (false).
+	 * If true, SupportedClasses is treated as a filter against the Parent Class of selected Blueprint assets.
+	 */
+	UPROPERTY(EditDefaultsOnly, Category="Asset Support")
+	bool bIsActionForBlueprints = false;
+
 	/**
 	 * The supported classes controls the list of classes that may be operated on by all of the asset functions in this
 	 * utility class.
+	 * @note When bIsActionForBlueprints is true, this will compare against the generated class of any Blueprint assets.
 	 */
 	UPROPERTY(EditDefaultsOnly, Category="Asset Support", meta=(AllowAbstract))
 	TArray<TSoftClassPtr<UObject>> SupportedClasses;
@@ -102,4 +111,4 @@ protected:
 	 */
 	UPROPERTY(EditDefaultsOnly, Category="Asset Support")
 	TArray<FAssetActionSupportCondition> SupportedConditions;
-};
\ No newline at end of file
+};
diff --git a/Engine/Source/Editor/Blutility/Private/AssetActionUtility.cpp b/Engine/Source/Editor/Blutility/Private/AssetActionUtility.cpp
--- a/Engine/Source/Editor/Blutility/Private/AssetActionUtility.cpp
+++ b/Engine/Source/Editor/Blutility/Private/AssetActionUtility.cpp
@@ -12,6 +12,11 @@
 
 #define LOCTEXT_NAMESPACE "AssetActionUtility"
 
+bool UAssetActionUtility::IsActionForBlueprints_Implementation() const
+{
+	return bIsActionForBlueprints;
+}
+
 void UAssetActionUtility::GetAssetRegistryTags(TArray<FAssetRegistryTag>& OutTags) const
 {
 	Super::GetAssetRegistryTags(OutTags);
