Add tooltip for Editor Utility actions when cooked

diff --git a/Engine/Source/Editor/Blutility/Private/BlutilityMenuExtensions.cpp b/Engine/Source/Editor/Blutility/Private/BlutilityMenuExtensions.cpp
--- a/Engine/Source/Editor/Blutility/Private/BlutilityMenuExtensions.cpp
+++ b/Engine/Source/Editor/Blutility/Private/BlutilityMenuExtensions.cpp
@@ -784,7 +784,7 @@ void FBlutilityMenuExtensions::CreateBlutilityActionsMenu(FMenuBuilder& MenuBuil
 
 			if (FunctionAndUtil.Util->GetUtilityBlueprintAsset().AssetClassPath == UBlueprintGeneratedClass::StaticClass()->GetClassPathName())
 			{
-				TooltipText = LOCTEXT("AssetUtilTooltip", "Click to execute");
+				TooltipText = FText::Format(LOCTEXT("AssetUtilTooltip", "{0}\n\n(Click to execute)"), FunctionAndUtil.FunctionData.TooltipText);
 			}
 			else if (FilterFailureMessage.IsEmpty())
 			{
