Add the "ContentBrowser.AssetContextMenu.ScriptedAssetActions" menu using UToolMenus

This will allow us to treat the "ScriptedAssetActions" sub-menu added to the right-click menu of assets in the Unreal editor like any other UToolMenu.

#rnx

diff --git a/Engine/Source/Editor/Blutility/Private/BlutilityContentBrowserExtensions.cpp b/Engine/Source/Editor/Blutility/Private/BlutilityContentBrowserExtensions.cpp
--- a/Engine/Source/Editor/Blutility/Private/BlutilityContentBrowserExtensions.cpp
+++ b/Engine/Source/Editor/Blutility/Private/BlutilityContentBrowserExtensions.cpp
@@ -10,6 +10,7 @@
 #include "Containers/Map.h"
 #include "Containers/Set.h"
 #include "ContentBrowserDelegates.h"
+#include "ContentBrowserMenuContexts.h"
 #include "ContentBrowserModule.h"
 #include "Delegates/Delegate.h"
 #include "EditorUtilityAssetPrototype.h"
@@ -27,21 +28,40 @@
 #include "Templates/SharedPointer.h"
 #include "Templates/SubclassOf.h"
 #include "Templates/UnrealTemplate.h"
+#include "ToolMenus.h"
 #include "UObject/Class.h"
 #include "UObject/NameTypes.h"
 #include "UObject/UObjectIterator.h"
 
 #define LOCTEXT_NAMESPACE "BlutilityContentBrowserExtensions"
 
-static FContentBrowserMenuExtender_SelectedAssets ContentBrowserExtenderDelegate;
-static FDelegateHandle ContentBrowserExtenderDelegateHandle;
+void FBlutilityContentBrowserExtensions::InstallHooks()
+{
+	UToolMenus::RegisterStartupCallback(
+		FSimpleMulticastDelegate::FDelegate::CreateStatic(&FBlutilityContentBrowserExtensions::RegisterMenus));
+}
 
-class FBlutilityContentBrowserExtensions_Impl
+void FBlutilityContentBrowserExtensions::RegisterMenus()
 {
-public:
-	static TSharedRef<FExtender> OnExtendContentBrowserAssetSelectionMenu(const TArray<FAssetData>& SelectedAssets)
+	// Mark us as the owner of everything we add.
+	FToolMenuOwnerScoped OwnerScoped("FBlutilityContentBrowserExtensions");
+
+	UToolMenu* const Menu = UToolMenus::Get()->ExtendMenu("ContentBrowser.AssetContextMenu");
+	if (!Menu)
 	{
-		TSharedRef<FExtender> Extender(new FExtender());
+		return;
+	}
+
+	FToolMenuSection& Section = Menu->FindOrAddSection("CommonAssetActions");
+	Section.AddDynamicEntry("BlutilityContentBrowserExtensions", FNewToolMenuSectionDelegate::CreateLambda([](FToolMenuSection& InSection)
+	{
+		UContentBrowserAssetContextMenuContext* const Context = InSection.FindContext<UContentBrowserAssetContextMenuContext>();
+		if (!Context)
+		{
+			return;
+		}
+
+		const TArray<FAssetData>& SelectedAssets = Context->SelectedAssets;
 
 		// Run thru the assets to determine if any meet our criteria
 		TMap<TSharedRef<FAssetActionUtilityPrototype>, TSet<int32>> UtilityAndSelectionIndices;
@@ -122,39 +142,16 @@ public:
 			EditorErrors.Notify(LOCTEXT("SomeProblemsWithAssetActionUtility", "There were some problems with some AssetActionUtility Blueprints."));
 		}
 
-		if (UtilityAndSelectionIndices.Num() > 0)
-		{
-			// Add asset actions extender
-			Extender->AddMenuExtension(
-				"CommonAssetActions",
-				EExtensionHook::After,
-				nullptr,
-				FMenuExtensionDelegate::CreateStatic(&FBlutilityMenuExtensions::CreateAssetBlutilityActionsMenu, MoveTemp(UtilityAndSelectionIndices), MoveTemp(SupportedAssets)));
-		}
-
-		return Extender;
-	}
-
-	static TArray<FContentBrowserMenuExtender_SelectedAssets>& GetExtenderDelegates()
-	{
-		FContentBrowserModule& ContentBrowserModule = FModuleManager::LoadModuleChecked<FContentBrowserModule>(TEXT("ContentBrowser"));
-		return ContentBrowserModule.GetAllAssetViewContextMenuExtenders();
-	}
-};
-
-void FBlutilityContentBrowserExtensions::InstallHooks()
-{
-	ContentBrowserExtenderDelegate = FContentBrowserMenuExtender_SelectedAssets::CreateStatic(&FBlutilityContentBrowserExtensions_Impl::OnExtendContentBrowserAssetSelectionMenu);
-
-	TArray<FContentBrowserMenuExtender_SelectedAssets>& CBMenuExtenderDelegates = FBlutilityContentBrowserExtensions_Impl::GetExtenderDelegates();
-	CBMenuExtenderDelegates.Add(ContentBrowserExtenderDelegate);
-	ContentBrowserExtenderDelegateHandle = CBMenuExtenderDelegates.Last().GetHandle();
+		FBlutilityMenuExtensions::CreateAssetBlutilityActionsMenu(InSection, MoveTemp(UtilityAndSelectionIndices), MoveTemp(SupportedAssets));
+	}));
 }
 
 void FBlutilityContentBrowserExtensions::RemoveHooks()
 {
-	TArray<FContentBrowserMenuExtender_SelectedAssets>& CBMenuExtenderDelegates = FBlutilityContentBrowserExtensions_Impl::GetExtenderDelegates();
-	CBMenuExtenderDelegates.RemoveAll([](const FContentBrowserMenuExtender_SelectedAssets& Delegate){ return Delegate.GetHandle() == ContentBrowserExtenderDelegateHandle; });
+	// Remove our startup delegate in case it's still around.
+	UToolMenus::UnRegisterStartupCallback("FBlutilityContentBrowserExtensions");
+	// Remove everything we added to UToolMenus.
+	UToolMenus::UnregisterOwner("FBlutilityContentBrowserExtensions");
 }
 
 #undef LOCTEXT_NAMESPACE
diff --git a/Engine/Source/Editor/Blutility/Private/BlutilityContentBrowserExtensions.h b/Engine/Source/Editor/Blutility/Private/BlutilityContentBrowserExtensions.h
--- a/Engine/Source/Editor/Blutility/Private/BlutilityContentBrowserExtensions.h
+++ b/Engine/Source/Editor/Blutility/Private/BlutilityContentBrowserExtensions.h
@@ -7,5 +7,6 @@ class FBlutilityContentBrowserExtensions
 {
 public:
 	static void InstallHooks();
+	static void RegisterMenus();
 	static void RemoveHooks();
 };
diff --git a/Engine/Source/Editor/Blutility/Private/BlutilityMenuExtensions.cpp b/Engine/Source/Editor/Blutility/Private/BlutilityMenuExtensions.cpp
--- a/Engine/Source/Editor/Blutility/Private/BlutilityMenuExtensions.cpp
+++ b/Engine/Source/Editor/Blutility/Private/BlutilityMenuExtensions.cpp
@@ -666,9 +666,10 @@ void FBlutilityMenuExtensions::CreateActorBlutilityActionsMenu(FMenuBuilder& Men
 	);
 }
 
-void FBlutilityMenuExtensions::CreateAssetBlutilityActionsMenu(FMenuBuilder& MenuBuilder, TMap<TSharedRef<FAssetActionUtilityPrototype>, TSet<int32>> Utils, const TArray<FAssetData> SelectedSupportedAssets)
+void FBlutilityMenuExtensions::CreateAssetBlutilityActionsMenu(FToolMenuSection& InSection, TMap<TSharedRef<FAssetActionUtilityPrototype>, TSet<int32>> Utils, const TArray<FAssetData> SelectedSupportedAssets)
 {
-	CreateBlutilityActionsMenu<FAssetData>(MenuBuilder, Utils,
+	CreateBlutilityActionsMenu<FAssetData>(InSection, Utils,
+		"ScriptedAssetActions",
 		LOCTEXT("ScriptedAssetActions", "Scripted Asset Actions"),
 		LOCTEXT("ScriptedAssetActionsTooltip", "Scripted actions available for the selected assets"),
 		[](const FProperty* Property) -> bool
@@ -954,6 +955,241 @@ void FBlutilityMenuExtensions::CreateBlutilityActionsMenu(FMenuBuilder& MenuBuil
 	}
 }
 
+template<typename SelectionType>
+void FBlutilityMenuExtensions::CreateBlutilityActionsMenu(FToolMenuSection& InSection, TMap<TSharedRef<FAssetActionUtilityPrototype>, TSet<int32>> Utils, const FName& MenuName, const FText& MenuLabel, const FText& MenuToolTip, TFunction<bool(const FProperty * Property)> IsValidPropertyType, const TArray<SelectionType> Selection, const FName& IconName)
+{
+	TMap<FString, TArray<FFunctionAndUtil>> CategoryFunctions;
+	ExtractFunctions(Utils, CategoryFunctions);
+	
+	auto AddFunctionEntries = [Selection, IsValidPropertyType](const TArray<FFunctionAndUtil>& FunctionUtils) -> TArray<FToolMenuEntry>
+	{
+		BlutilityUtil::FTextFilterExpressionContext TextFilterContext;
+		
+		FTextFilterExpressionEvaluator TextFilterExpressionEvaluator(ETextFilterExpressionEvaluatorMode::Complex);
+		TArray<FToolMenuEntry> GeneratedMenuEntries;
+
+		for (const FFunctionAndUtil& FunctionAndUtil : FunctionUtils)
+		{
+			bool PassesFilterCondition = true;
+			bool bShowInMenu = true;
+			FString FilterFailureMessage;
+			
+			if constexpr ( std::is_same_v<FAssetData, SelectionType> )
+			{
+				TArray<FAssetActionSupportCondition> Conditions = FunctionAndUtil.Util->GetAssetActionSupportConditions();
+
+				for (const FAssetActionSupportCondition& Condition : Conditions)
+				{
+					TextFilterExpressionEvaluator.SetFilterText(FText::FromString(Condition.Filter));
+				
+					for (const int32& SelectionIndex : FunctionAndUtil.SelectionIndices)
+					{
+						const auto SelectedAsset = Selection[SelectionIndex];
+						TextFilterContext.SetAsset(&SelectedAsset);
+						if (!TextFilterExpressionEvaluator.TestTextFilter(TextFilterContext))
+						{
+							PassesFilterCondition = false;
+							FilterFailureMessage = Condition.FailureReason;
+							break;
+						}
+					}
+
+					if (!PassesFilterCondition)
+					{
+						bShowInMenu = Condition.bShowInMenuIfFilterFails;
+						break;
+					}
+				}
+			}
+
+			if (!PassesFilterCondition && !bShowInMenu)
+			{
+				continue;
+			}
+			
+			FText TooltipText;
+
+			if (FunctionAndUtil.Util->GetUtilityBlueprintAsset().AssetClassPath == UBlueprintGeneratedClass::StaticClass()->GetClassPathName())
+			{
+				TooltipText = LOCTEXT("AssetUtilTooltip", "Click to execute");
+			}
+			else if (FilterFailureMessage.IsEmpty())
+			{
+				TooltipText = FText::Format(LOCTEXT("AssetUtilTooltipFormat", "{0}\n\n(Shift-click to edit script)"), FunctionAndUtil.FunctionData.TooltipText);	
+			}
+			else
+			{
+				TooltipText = FText::Format(LOCTEXT("AssetUtilTooltipWithErrorFormat", "{0}\n\n({1})\n\n(Shift-click to edit script)"), FunctionAndUtil.FunctionData.TooltipText, FText::FromString(FilterFailureMessage));
+			}
+
+			FExecuteAction ExecuteAction = FExecuteAction::CreateLambda([FunctionAndUtil, Selection, IsValidPropertyType]
+			{
+				if (FSlateApplication::Get().GetModifierKeys().IsShiftDown())
+				{
+					OpenEditorForUtility(FunctionAndUtil);
+				}
+				else
+				{
+					// We dont run this on the CDO, as bad things could occur!
+					UObject* TempObject = NewObject<UObject>(GetTransientPackage(), Cast<UObject>(FunctionAndUtil.Util->LoadUtilityAsset())->GetClass());
+					TempObject->AddToRoot(); // Some Blutility actions might run GC so the TempObject needs to be rooted to avoid getting destroyed
+
+					UFunction* Function = FunctionAndUtil.GetFunction();
+					
+					if (Function->NumParms > 0)
+					{
+						// Create a parameter struct and fill in defaults
+						TSharedRef<FStructOnScope> FuncParams = MakeShared<FStructOnScope>(Function);
+
+						FProperty* FirstParamProperty = nullptr;
+
+						int32 ParameterIndex = 0;
+						for (TFieldIterator<FProperty> It(Function); It&& It->HasAnyPropertyFlags(CPF_Parm); ++It)
+						{
+							FString Defaults;
+							if (UEdGraphSchema_K2::FindFunctionParameterDefaultValue(Function, *It, Defaults))
+							{
+								It->ImportText_Direct(*Defaults, It->ContainerPtrToValuePtr<uint8>(FuncParams->GetStructMemory()), nullptr, PPF_None);
+							}
+
+							// Check to see if the first parameter matches the selection object type, in that case we can directly forward the selection to it
+							if (ParameterIndex == 0 && IsValidPropertyType(*It))
+							{
+								FirstParamProperty = *It;
+							}
+
+							++ParameterIndex;
+						}
+
+						bool bApply = true;
+
+						if (!FirstParamProperty || ParameterIndex > 1)
+						{
+							// pop up a dialog to input params to the function
+							TSharedRef<SWindow> Window = SNew(SWindow)
+								.Title(Function->GetDisplayNameText())
+								.ClientSize(FVector2D(400, 200))
+								.SupportsMinimize(false)
+								.SupportsMaximize(false);
+
+							TSharedPtr<SFunctionParamDialog> Dialog;
+							Window->SetContent(
+								SAssignNew(Dialog, SFunctionParamDialog, Window, FuncParams, FirstParamProperty ? FirstParamProperty->GetFName() : NAME_None)
+								.OkButtonText(LOCTEXT("OKButton", "OK"))
+								.OkButtonTooltipText(Function->GetToolTipText()));
+
+							GEditor->EditorAddModalWindow(Window);
+							bApply = Dialog->bOKPressed;
+						}
+
+
+						if (bApply)
+						{
+							FScopedTransaction Transaction(NSLOCTEXT("UnrealEd", "BlutilityAction", "Blutility Action"));
+							FEditorScriptExecutionGuard ScriptGuard;
+							const bool bForwardUserSelection = FirstParamProperty != nullptr;
+							if (bForwardUserSelection)
+							{
+								// For each user-select asset forward the selection object into the function first's parameter (if it matches)
+								const FString Path = FirstParamProperty->GetPathName(Function);
+
+								// Ensure we only process selection objects that are valid for this function/utility
+								for (const int32& SelectionIndex : FunctionAndUtil.SelectionIndices)
+								{
+									const auto SelectedAsset = Selection[SelectionIndex];
+									FirstParamProperty->CopySingleValue(FirstParamProperty->ContainerPtrToValuePtr<uint8>(FuncParams->GetStructMemory()), &SelectedAsset);
+									TempObject->ProcessEvent(Function, FuncParams->GetStructMemory());
+								}
+							}
+							else
+							{
+								// User is expected to manage the asset selection on its own
+								TempObject->ProcessEvent(Function, FuncParams->GetStructMemory());
+							}
+						}
+					}
+					else
+					{
+						FScopedTransaction Transaction(NSLOCTEXT("UnrealEd", "BlutilityAction", "Blutility Action"));
+						FEditorScriptExecutionGuard ScriptGuard;
+						TempObject->ProcessEvent(Function, nullptr);
+					}
+
+					TempObject->RemoveFromRoot();
+				}
+			});
+
+			FCanExecuteAction CanExecuteAction = FCanExecuteAction::CreateLambda([PassesFilterCondition]()
+			{
+				return PassesFilterCondition;
+			});
+
+			FToolMenuEntry MenuEntry = FToolMenuEntry::InitMenuEntry(
+				FunctionAndUtil.FunctionData.Name,
+				FunctionAndUtil.FunctionData.NameText,
+				TooltipText,
+				FSlateIcon(FAppStyle::GetAppStyleSetName(), "GraphEditor.Event_16x"),
+				FUIAction(ExecuteAction, CanExecuteAction)
+			);
+
+			GeneratedMenuEntries.Emplace(MenuEntry);
+		}
+
+		return GeneratedMenuEntries;
+	};
+
+	// Add a menu item for each function
+	if (CategoryFunctions.Num() > 0)
+	{
+		InSection.AddSubMenu(
+			MenuName,
+			MenuLabel,
+			MenuToolTip,
+			FNewToolMenuDelegate::CreateLambda([CategoryFunctions, AddFunctionEntries](UToolMenu* InMenu)
+			{
+				TArray<FString> CategoryNames;
+				CategoryFunctions.GenerateKeyArray(CategoryNames);
+				CategoryNames.Remove(FString());
+				CategoryNames.Sort();
+
+				FToolMenuSection& UnnamedSection = InMenu->FindOrAddSection(NAME_None);
+				
+				// Add functions belong to the same category to a sub-menu
+				for (const FString& CategoryName : CategoryNames)
+				{
+					const TArray<FToolMenuEntry> GeneratedMenuEntries = AddFunctionEntries(CategoryFunctions.FindChecked(CategoryName));
+					if (GeneratedMenuEntries.Num() > 0)
+					{
+						UnnamedSection.AddSubMenu(
+							NAME_None,
+							FText::FromString(CategoryName),
+							FText::FromString(CategoryName),
+							FNewToolMenuDelegate::CreateLambda([GeneratedMenuEntries](UToolMenu* InCategorySubMenu)
+							{
+								for (const FToolMenuEntry& MenuEntry : GeneratedMenuEntries)
+								{
+									InCategorySubMenu->AddMenuEntry(NAME_None, MenuEntry);
+								}
+							}
+						));
+					}
+				}
+
+				// Non-categorized functions
+				const TArray<FFunctionAndUtil>* DefaultCategoryFunctionsPtr = CategoryFunctions.Find(FString());
+				if (DefaultCategoryFunctionsPtr)
+				{
+					for (const FToolMenuEntry& MenuEntry : AddFunctionEntries(*DefaultCategoryFunctionsPtr))
+					{
+						InMenu->AddMenuEntry(NAME_None, MenuEntry);
+					}
+				}
+			}),
+			false,
+			FSlateIcon(FAppStyle::GetAppStyleSetName(), IconName)
+		);
+	}
+}
 
 UEditorUtilityExtension::UEditorUtilityExtension(const class FObjectInitializer& ObjectInitializer)
 	: Super(ObjectInitializer)
diff --git a/Engine/Source/Editor/Blutility/Private/BlutilityMenuExtensions.h b/Engine/Source/Editor/Blutility/Private/BlutilityMenuExtensions.h
--- a/Engine/Source/Editor/Blutility/Private/BlutilityMenuExtensions.h
+++ b/Engine/Source/Editor/Blutility/Private/BlutilityMenuExtensions.h
@@ -21,6 +21,7 @@ class FProperty;
 class FString;
 class FText;
 class IEditorUtilityExtension;
+struct FToolMenuSection;
 class UFunction;
 class UObject;
 class FAssetActionUtilityPrototype;
@@ -33,7 +34,7 @@ public:
 	static void GetBlutilityClasses(TArray<FAssetData>& OutAssets, FTopLevelAssetPath InClassName);
 
 	/** Helper function that populates a menu based on the exposed functions in a set of Blutility objects */
-	static void CreateAssetBlutilityActionsMenu(FMenuBuilder& MenuBuilder, TMap<TSharedRef<FAssetActionUtilityPrototype>, TSet<int32>> Utils, const TArray<FAssetData> SelectedSupportedAssets);
+	static void CreateAssetBlutilityActionsMenu(FToolMenuSection& InSection, TMap<TSharedRef<FAssetActionUtilityPrototype>, TSet<int32>> Utils, const TArray<FAssetData> SelectedSupportedAssets);
 	static void CreateActorBlutilityActionsMenu(FMenuBuilder& MenuBuilder, TMap<TSharedRef<FAssetActionUtilityPrototype>, TSet<int32>> Utils, const TArray<AActor*> SelectedSupportedActors);
 	
 protected:
@@ -73,6 +74,8 @@ protected:
 protected:
 	template<typename SelectionType>
 	static void CreateBlutilityActionsMenu(FMenuBuilder& MenuBuilder, TMap<TSharedRef<FAssetActionUtilityPrototype>, TSet<int32>> Utils, const FText& MenuLabel, const FText& MenuToolTip, TFunction<bool(const FProperty* Property)> IsValidPropertyType, const TArray<SelectionType> Selection, const FName& IconName = "GraphEditor.Event_16x");
+	template<typename SelectionType>
+	static void CreateBlutilityActionsMenu(FToolMenuSection& InSection, TMap<TSharedRef<FAssetActionUtilityPrototype>, TSet<int32>> Utils, const FName& MenuName, const FText& MenuLabel, const FText& MenuToolTip, TFunction<bool(const FProperty* Property)> IsValidPropertyType, const TArray<SelectionType> Selection, const FName& IconName = "GraphEditor.Event_16x");
 
 	static void OpenEditorForUtility(const FFunctionAndUtil& FunctionAndUtil);
 	static void ExtractFunctions(TMap<TSharedRef<FAssetActionUtilityPrototype>, TSet<int32>>& Utils, TMap<FString, TArray<FFunctionAndUtil>>& OutCategoryFunctions);
