Add the "LevelEditor.ActorContextMenu.ScriptedAssetActions" menu using UToolMenus

This will allow us to treat the "ScriptedActorActions" sub-menu added to the right-click menu of actors in the Level Editor like any other UToolMenu.

#rnx

diff --git a/Engine/Source/Editor/Blutility/Private/BlutilityLevelEditorExtensions.cpp b/Engine/Source/Editor/Blutility/Private/BlutilityLevelEditorExtensions.cpp
--- a/Engine/Source/Editor/Blutility/Private/BlutilityLevelEditorExtensions.cpp
+++ b/Engine/Source/Editor/Blutility/Private/BlutilityLevelEditorExtensions.cpp
@@ -21,9 +21,11 @@
 #include "Logging/MessageLog.h"
 #include "Misc/NamePermissionList.h"
 #include "Modules/ModuleManager.h"
+#include "Selection.h"
 #include "Templates/Casts.h"
 #include "Templates/SharedPointer.h"
 #include "Templates/SubclassOf.h"
+#include "ToolMenus.h"
 #include "UObject/Class.h"
 #include "UObject/NameTypes.h"
 #include "UObject/UObjectIterator.h"
@@ -32,18 +34,32 @@ class FUICommandList;
 
 #define LOCTEXT_NAMESPACE "BlutilityLevelEditorExtensions"
 
-FDelegateHandle LevelViewportExtenderHandle;
+void FBlutilityLevelEditorExtensions::InstallHooks()
+{
+	UToolMenus::RegisterStartupCallback(
+		FSimpleMulticastDelegate::FDelegate::CreateStatic(&FBlutilityLevelEditorExtensions::RegisterMenus));
+}
 
-class FBlutilityLevelEditorExtensions_Impl
+void FBlutilityLevelEditorExtensions::RegisterMenus()
 {
-public:
-	static TSharedRef<FExtender> OnExtendLevelEditorActorContextMenu(const TSharedRef<FUICommandList> CommandList, const TArray<AActor*> SelectedActors)
+	// Mark us as the owner of everything we add.
+	FToolMenuOwnerScoped OwnerScoped("FBlutilityLevelEditorExtensions");
+
+	UToolMenu* const Menu = UToolMenus::Get()->ExtendMenu("LevelEditor.ActorContextMenu");
+	if (!Menu)
 	{
-		TSharedRef<FExtender> Extender(new FExtender());
-			   
+		return;
+	}
+
+	FToolMenuSection& Section = Menu->FindOrAddSection("ActorOptions");
+	Section.AddDynamicEntry("BlutilityLevelEditorExtensions", FNewToolMenuSectionDelegate::CreateLambda([](FToolMenuSection& InSection)
+	{
+		TArray<AActor*> SelectedActors;
+		GEditor->GetSelectedActors()->GetSelectedObjects<AActor>(SelectedActors);
+
 		TMap<TSharedRef<FAssetActionUtilityPrototype>, TSet<int32>> UtilityAndSelectionIndices;
 
-		// Run thru the actors to determine if any meet our criteria
+		// Run through the actors to determine if any meet our criteria.
 		TArray<AActor*> SupportedActors;
 		if (SelectedActors.Num() > 0)
 		{
@@ -108,38 +124,17 @@ public:
 
 		if (SupportedActors.Num() > 0)
 		{
-			// Add asset actions extender
-			Extender->AddMenuExtension(
-				"ActorOptions",
-				EExtensionHook::After,
-				CommandList,
-				FMenuExtensionDelegate::CreateStatic(&FBlutilityMenuExtensions::CreateActorBlutilityActionsMenu, UtilityAndSelectionIndices, SupportedActors));
+			FBlutilityMenuExtensions::CreateActorBlutilityActionsMenu(InSection, UtilityAndSelectionIndices, SupportedActors);
 		}
-
-		return Extender;
-	}
-};
-
-void FBlutilityLevelEditorExtensions::InstallHooks()
-{
-	FLevelEditorModule& LevelEditorModule = FModuleManager::Get().LoadModuleChecked<FLevelEditorModule>("LevelEditor");
-	auto& MenuExtenders = LevelEditorModule.GetAllLevelViewportContextMenuExtenders();
-
-	MenuExtenders.Add(FLevelEditorModule::FLevelViewportMenuExtender_SelectedActors::CreateStatic(&FBlutilityLevelEditorExtensions_Impl::OnExtendLevelEditorActorContextMenu));
-	LevelViewportExtenderHandle = MenuExtenders.Last().GetHandle();
+	}));
 }
 
 void FBlutilityLevelEditorExtensions::RemoveHooks()
 {
-	if (LevelViewportExtenderHandle.IsValid())
-	{
-		FLevelEditorModule* LevelEditorModule = FModuleManager::Get().GetModulePtr<FLevelEditorModule>("LevelEditor");
-		if (LevelEditorModule)
-		{
-			typedef FLevelEditorModule::FLevelViewportMenuExtender_SelectedActors DelegateType;
-			LevelEditorModule->GetAllLevelViewportContextMenuExtenders().RemoveAll([=](const DelegateType& In) { return In.GetHandle() == LevelViewportExtenderHandle; });
-		}
-	}
+	// Remove our startup delegate in case it's still around.
+	UToolMenus::UnRegisterStartupCallback("FBlutilityLevelEditorExtensions");
+	// Remove everything we added to UToolMenus.
+	UToolMenus::UnregisterOwner("FBlutilityLevelEditorExtensions");
 }
 
 #undef LOCTEXT_NAMESPACE
diff --git a/Engine/Source/Editor/Blutility/Private/BlutilityLevelEditorExtensions.h b/Engine/Source/Editor/Blutility/Private/BlutilityLevelEditorExtensions.h
--- a/Engine/Source/Editor/Blutility/Private/BlutilityLevelEditorExtensions.h
+++ b/Engine/Source/Editor/Blutility/Private/BlutilityLevelEditorExtensions.h
@@ -7,5 +7,6 @@ class FBlutilityLevelEditorExtensions
 {
 public:
 	static void InstallHooks();
+	static void RegisterMenus();
 	static void RemoveHooks();
 };
diff --git a/Engine/Source/Editor/Blutility/Private/BlutilityMenuExtensions.cpp b/Engine/Source/Editor/Blutility/Private/BlutilityMenuExtensions.cpp
--- a/Engine/Source/Editor/Blutility/Private/BlutilityMenuExtensions.cpp
+++ b/Engine/Source/Editor/Blutility/Private/BlutilityMenuExtensions.cpp
@@ -647,9 +647,10 @@ void FBlutilityMenuExtensions::GetBlutilityClasses(TArray<FAssetData>& OutAssets
 	}
 }
 
-void FBlutilityMenuExtensions::CreateActorBlutilityActionsMenu(FMenuBuilder& MenuBuilder, TMap<TSharedRef<FAssetActionUtilityPrototype>, TSet<int32>> Utils, const TArray<AActor*> SelectedSupportedActors)
+void FBlutilityMenuExtensions::CreateActorBlutilityActionsMenu(FToolMenuSection& InSection, TMap<TSharedRef<FAssetActionUtilityPrototype>, TSet<int32>> Utils, const TArray<AActor*> SelectedSupportedActors)
 {
-	CreateBlutilityActionsMenu<AActor*>(MenuBuilder, Utils,
+	CreateBlutilityActionsMenu<AActor*>(InSection, Utils,
+		"ScriptedActorActions",
 		LOCTEXT("ScriptedActorActions", "Scripted Actor Actions"),
 		LOCTEXT("ScriptedActorActionsTooltip", "Scripted actions available for the selected actors"),
 		[](const FProperty* Property) -> bool
diff --git a/Engine/Source/Editor/Blutility/Private/BlutilityMenuExtensions.h b/Engine/Source/Editor/Blutility/Private/BlutilityMenuExtensions.h
--- a/Engine/Source/Editor/Blutility/Private/BlutilityMenuExtensions.h
+++ b/Engine/Source/Editor/Blutility/Private/BlutilityMenuExtensions.h
@@ -34,8 +34,8 @@ public:
 	static void GetBlutilityClasses(TArray<FAssetData>& OutAssets, FTopLevelAssetPath InClassName);
 
 	/** Helper function that populates a menu based on the exposed functions in a set of Blutility objects */
+	static void CreateActorBlutilityActionsMenu(FToolMenuSection& InSection, TMap<TSharedRef<FAssetActionUtilityPrototype>, TSet<int32>> Utils, const TArray<AActor*> SelectedSupportedActors);
 	static void CreateAssetBlutilityActionsMenu(FToolMenuSection& InSection, TMap<TSharedRef<FAssetActionUtilityPrototype>, TSet<int32>> Utils, const TArray<FAssetData> SelectedSupportedAssets);
-	static void CreateActorBlutilityActionsMenu(FMenuBuilder& MenuBuilder, TMap<TSharedRef<FAssetActionUtilityPrototype>, TSet<int32>> Utils, const TArray<AActor*> SelectedSupportedActors);
 	
 protected:
 	// Helper struct to track the util to call a function on
