Speculative fix editor error notifications with generated classes.

diff --git a/Engine/Source/Editor/Blutility/Private/BlutilityContentBrowserExtensions.cpp b/Engine/Source/Editor/Blutility/Private/BlutilityContentBrowserExtensions.cpp
--- a/Engine/Source/Editor/Blutility/Private/BlutilityContentBrowserExtensions.cpp
+++ b/Engine/Source/Editor/Blutility/Private/BlutilityContentBrowserExtensions.cpp
@@ -15,6 +15,7 @@
 #include "Delegates/Delegate.h"
 #include "EditorUtilityAssetPrototype.h"
 #include "EditorUtilityBlueprint.h"
+#include "EditorUtilityWidgetProjectSettings.h"
 #include "Engine/Blueprint.h"
 #include "Engine/BlueprintGeneratedClass.h"
 #include "Framework/MultiBox/MultiBoxExtender.h"
@@ -139,7 +140,12 @@ void FBlutilityContentBrowserExtensions::RegisterMenus()
 				}
 			}
 
-			EditorErrors.Notify(LOCTEXT("SomeProblemsWithAssetActionUtility", "There were some problems with some AssetActionUtility Blueprints."));
+			// Don't warn errors if searching generated classes, since not all utilities may be updated to work with generated classes yet (Must be done piecemeal).
+			const UEditorUtilityWidgetProjectSettings* EditorUtilitySettings = GetDefault<UEditorUtilityWidgetProjectSettings>();
+			if (!EditorUtilitySettings->bSearchGeneratedClassesForScriptedActions)
+			{
+				EditorErrors.Notify(LOCTEXT("SomeProblemsWithAssetActionUtility", "There were some problems with some AssetActionUtility Blueprints."));
+			}
 		}
 
 		FBlutilityMenuExtensions::CreateAssetBlutilityActionsMenu(InSection, MoveTemp(UtilityAndSelectionIndices), MoveTemp(SupportedAssets));
diff --git a/Engine/Source/Editor/Blutility/Private/BlutilityLevelEditorExtensions.cpp b/Engine/Source/Editor/Blutility/Private/BlutilityLevelEditorExtensions.cpp
--- a/Engine/Source/Editor/Blutility/Private/BlutilityLevelEditorExtensions.cpp
+++ b/Engine/Source/Editor/Blutility/Private/BlutilityLevelEditorExtensions.cpp
@@ -13,6 +13,7 @@
 #include "Delegates/Delegate.h"
 #include "EditorUtilityAssetPrototype.h"
 #include "EditorUtilityBlueprint.h"
+#include "EditorUtilityWidgetProjectSettings.h"
 #include "Framework/MultiBox/MultiBoxExtender.h"
 #include "GameFramework/Actor.h"
 #include "HAL/PlatformCrt.h"
@@ -118,8 +119,13 @@ void FBlutilityLevelEditorExtensions::RegisterMenus()
 					}
 				}
 			}
-			
-			EditorErrors.Notify(LOCTEXT("SomeProblemsWithActorActionUtility", "There were some problems with some ActorActionUtility Blueprints."));
+
+			// Don't warn errors if searching generated classes, since not all utilities may be updated to work with generated classes yet (Must be done piecemeal).
+			const UEditorUtilityWidgetProjectSettings* EditorUtilitySettings = GetDefault<UEditorUtilityWidgetProjectSettings>();
+			if (!EditorUtilitySettings->bSearchGeneratedClassesForScriptedActions)
+			{
+				EditorErrors.Notify(LOCTEXT("SomeProblemsWithActorActionUtility", "There were some problems with some ActorActionUtility Blueprints."));
+			}
 		}
 
 		if (SupportedActors.Num() > 0)
