Add EditorFunctionLibrary, so that editor utility blueprints can make libraries of static functions

diff --git a/Engine/Source/Editor/Blutility/Private/EditorFunctionLibrary.h b/Engine/Source/Editor/Blutility/Private/EditorFunctionLibrary.h
new file mode 100644
--- /dev/null
+++ b/Engine/Source/Editor/Blutility/Private/EditorFunctionLibrary.h
@@ -0,0 +1,23 @@
+// Copyright Epic Games, Inc. All Rights Reserved.
+
+#pragma once
+
+#include "CoreMinimal.h"
+#include "Kismet/BlueprintFunctionLibrary.h"
+#include "EditorFunctionLibrary.generated.h"
+
+/**
+ * Library of static functions that can use the editor APIs
+ */
+UCLASS(Abstract, MinimalAPI)
+class UEditorFunctionLibrary : public UBlueprintFunctionLibrary
+{
+	GENERATED_BODY()
+};
+
+/*
+	Note: There is no reason to extend this in c++, it is used as a sentinel type for blueprints
+	that want to provide static reusable functions for editing logic. In c++ you can just extend
+	UBlueprintFunctionLibrary as normal and put your class in an uncooked or editor only module
+	or WITH_EDITOR block.
+*/
diff --git a/Engine/Source/Editor/Blutility/Private/EditorUtilityBlueprintFactory.cpp b/Engine/Source/Editor/Blutility/Private/EditorUtilityBlueprintFactory.cpp
--- a/Engine/Source/Editor/Blutility/Private/EditorUtilityBlueprintFactory.cpp
+++ b/Engine/Source/Editor/Blutility/Private/EditorUtilityBlueprintFactory.cpp
@@ -11,6 +11,7 @@
 #include "EditorUtilityBlueprint.h"
 #include "EditorUtilityObject.h"
 #include "EditorUtilityWidget.h"
+#include "EditorFunctionLibrary.h"
 #include "Engine/Blueprint.h"
 #include "Engine/BlueprintGeneratedClass.h"
 #include "HAL/Platform.h"
@@ -90,6 +91,7 @@ bool UEditorUtilityBlueprintFactory::ConfigureProperties()
 	Options.ExtraPickerCommonClasses.Add(AEditorUtilityActor::StaticClass());
 	Options.ExtraPickerCommonClasses.Add(UEditorUtilityObject::StaticClass());
 	Options.ExtraPickerCommonClasses.Add(UAssetActionUtility::StaticClass());
+	Options.ExtraPickerCommonClasses.Add(UEditorFunctionLibrary::StaticClass());
 
 	TSharedPtr< FBlutilityBlueprintFactoryFilter > Filter = MakeShareable(new FBlutilityBlueprintFactoryFilter);
 	Filter->DisallowedChildOfClasses.Add(UEditorUtilityWidget::StaticClass());
@@ -111,17 +113,20 @@ UObject* UEditorUtilityBlueprintFactory::FactoryCreateNew(UClass* Class, UObject
 	// Make sure we are trying to factory a blueprint, then create and init one
 	check(Class->IsChildOf(UBlueprint::StaticClass()));
 
-	if ((ParentClass == NULL) || !FKismetEditorUtilities::CanCreateBlueprintOfClass(ParentClass))
+	EBlueprintType BPType = BPTYPE_Normal;
+	if (ParentClass == UEditorFunctionLibrary::StaticClass())
+	{
+		BPType = BPTYPE_FunctionLibrary;
+	}
+	else if ((ParentClass == NULL) || !FKismetEditorUtilities::CanCreateBlueprintOfClass(ParentClass))
 	{
 		FFormatNamedArguments Args;
 		Args.Add( TEXT("ClassName"), (ParentClass != NULL) ? FText::FromString( ParentClass->GetName() ) : NSLOCTEXT("UnrealEd", "Null", "(null)") );
-		FMessageDialog::Open( EAppMsgType::Ok, FText::Format( NSLOCTEXT("UnrealEd", "CannotCreateBlueprintFromClass", "Cannot create a blueprint based on the class '{0}'."), Args ) );
+		FMessageDialog::Open( EAppMsgType::Ok, FText::Format( NSLOCTEXT("UnrealEd", "CannotCreateBlueprintFromClass", "Cannot create a blueprint based on the class '{ClassName}'."), Args ) );
 		return NULL;
 	}
-	else
-	{
-		return FKismetEditorUtilities::CreateBlueprint(ParentClass, InParent, Name, BPTYPE_Normal, UEditorUtilityBlueprint::StaticClass(), UBlueprintGeneratedClass::StaticClass());
-	}
+
+	return FKismetEditorUtilities::CreateBlueprint(ParentClass, InParent, Name, BPType, UEditorUtilityBlueprint::StaticClass(), UBlueprintGeneratedClass::StaticClass());
 }
 
 bool UEditorUtilityBlueprintFactory::CanCreateNew() const
