Blutility truncation warning fixes

diff --git a/Engine/Source/Editor/Blutility/Private/AsyncImageExport.cpp b/Engine/Source/Editor/Blutility/Private/AsyncImageExport.cpp
--- a/Engine/Source/Editor/Blutility/Private/AsyncImageExport.cpp
+++ b/Engine/Source/Editor/Blutility/Private/AsyncImageExport.cpp
@@ -96,21 +96,21 @@ void UAsyncImageExport::Activate()
 		TArray64<uint8> TextureRawData;
 		verify( TextureSource.GetMipData(TextureRawData, 0) );
 
-		const int64 BytesPerPixel = TextureSource.GetBytesPerPixel();
+		const int32 BytesPerPixel = static_cast<int32>(TextureSource.GetBytesPerPixel());
 		const ETextureSourceFormat SourceFormat = TextureSource.GetFormat();
 
-		const int64 Width = TextureSource.GetSizeX();
-		const int64 Height = TextureSource.GetSizeY();
+		const int32 Width = static_cast<int32>(TextureSource.GetSizeX());
+		const int32 Height = static_cast<int32>(TextureSource.GetSizeY());
 
 		TArray<FColor> OutPixels;
 		OutPixels.SetNumZeroed(Width * Height);
 
 		FColor Color(0, 0, 0, 0);
-		for (int64 Y = 0; Y < Height; ++Y)
+		for (int32 Y = 0; Y < Height; ++Y)
 		{
-			for (int64 X = 0; X < Width; ++X)
+			for (int32 X = 0; X < Width; ++X)
 			{
-				const int64 PixelByteOffset = (X + Y * Width) * BytesPerPixel;
+				const int32 PixelByteOffset = (X + Y * Width) * BytesPerPixel;
 				const uint8* PixelPtr = TextureRawData.GetData() + PixelByteOffset;
 
 				switch (SourceFormat)
