Removed a bad assertion in UEditorUtilitySubsystem::RemoveTaskFromActiveList, and update a similar check in RegisterAndExecuteTask to ensure we avoid duplicate task registration

diff --git a/Engine/Source/Editor/Blutility/Private/EditorUtilitySubsystem.cpp b/Engine/Source/Editor/Blutility/Private/EditorUtilitySubsystem.cpp
--- a/Engine/Source/Editor/Blutility/Private/EditorUtilitySubsystem.cpp
+++ b/Engine/Source/Editor/Blutility/Private/EditorUtilitySubsystem.cpp
@@ -529,6 +529,10 @@ void UEditorUtilitySubsystem::RegisterAndExecuteTask(UEditorUtilityTask* NewTask
 
 		// Register it
 		check(!(PendingTasks.Contains(NewTask) || ActiveTaskStack.Contains(NewTask)));
+		for (auto& KVP : PendingTasks)
+		{
+			check(!KVP.Value.Contains(NewTask));
+		}
 		NewTask->MyTaskManager = this;
 		NewTask->MyParentTask = OptionalParentTask;
 
@@ -544,7 +548,6 @@ void UEditorUtilitySubsystem::RemoveTaskFromActiveList(UEditorUtilityTask* Task)
 	{
 		if (ensure(Task->MyTaskManager == this))
 		{
-			check(PendingTasks.Contains(Task) || ActiveTaskStack.Contains(Task));
 			PendingTasks.Remove(Task);
 			ActiveTaskStack.Remove(Task);
 
