EditorUtilityBlueprintFactory: Modify the class picker to use a TreeView to show classes in hierarchy

diff --git a/Engine/Source/Editor/Blutility/Private/EditorUtilityBlueprintFactory.cpp b/Engine/Source/Editor/Blutility/Private/EditorUtilityBlueprintFactory.cpp
--- a/Engine/Source/Editor/Blutility/Private/EditorUtilityBlueprintFactory.cpp
+++ b/Engine/Source/Editor/Blutility/Private/EditorUtilityBlueprintFactory.cpp
@@ -81,7 +81,7 @@ bool UEditorUtilityBlueprintFactory::ConfigureProperties()
 	// Fill in options
 	FClassViewerInitializationOptions Options;
 	Options.Mode = EClassViewerMode::ClassPicker;
-	Options.DisplayMode = EClassViewerDisplayMode::ListView;
+	Options.DisplayMode = EClassViewerDisplayMode::TreeView;
 	// Only want blueprint actor base classes.
 	Options.bIsBlueprintBaseOnly = true;
 	// This will allow unloaded blueprints to be shown.
