Allow overriding the name of an EditorUtilityTask

diff --git a/Engine/Source/Editor/Blutility/Classes/EditorUtilityTask.h b/Engine/Source/Editor/Blutility/Classes/EditorUtilityTask.h
--- a/Engine/Source/Editor/Blutility/Classes/EditorUtilityTask.h
+++ b/Engine/Source/Editor/Blutility/Classes/EditorUtilityTask.h
@@ -57,6 +57,7 @@ public:
 protected:
 	virtual void BeginExecution() {}
 	virtual void CancelRequested() {}
+	virtual FText GetTaskTitle() const;
 
 	UFUNCTION(BlueprintImplementableEvent, Category=Task, meta=(DisplayName="BeginExecution"))
 	void ReceiveBeginExecution();
@@ -64,7 +65,11 @@ protected:
 	UFUNCTION(BlueprintImplementableEvent, Category=Task, meta=(DisplayName="CancelRequested"))
 	void ReceiveCancelRequested();
 
+	UFUNCTION(BlueprintImplementableEvent, Category=Task)
+	FText GetTaskTitleOverride() const;
+
 private:
+	// Calls GetTaskTitle() and GetTaskTitleOverride()
 	void CreateNotification();
 
 	// Calls BeginExecution() and ReceiveBeginExecution()
diff --git a/Engine/Source/Editor/Blutility/Private/EditorUtilityTask.cpp b/Engine/Source/Editor/Blutility/Private/EditorUtilityTask.cpp
--- a/Engine/Source/Editor/Blutility/Private/EditorUtilityTask.cpp
+++ b/Engine/Source/Editor/Blutility/Private/EditorUtilityTask.cpp
@@ -78,10 +78,25 @@ void UEditorUtilityTask::FinishExecutingTask()
 	OnFinished.Broadcast(this);
 }
 
+FText UEditorUtilityTask::GetTaskTitle() const
+{
+	return GetClass()->GetDisplayNameText();
+}
+
 void UEditorUtilityTask::CreateNotification()
 {
+	FText TaskTitle = GetTaskTitleOverride();
+	if (TaskTitle.IsEmpty())
+	{
+		TaskTitle = GetTaskTitle();
+	}
+	if (TaskTitle.IsEmpty())
+	{
+		TaskTitle = GetClass()->GetDisplayNameText();
+	}
+
 	FAsyncTaskNotificationConfig NotificationConfig;
-	NotificationConfig.TitleText = FText::Format(LOCTEXT("NotificationEditorUtilityTaskTitle", "Task {0}"), GetClass()->GetDisplayNameText());
+	NotificationConfig.TitleText = FText::Format(LOCTEXT("NotificationEditorUtilityTaskTitle", "Task {0}"), TaskTitle);
 	NotificationConfig.ProgressText = LOCTEXT("Running", "Running");
 	NotificationConfig.bCanCancel = true;
 	TaskNotification = MakeUnique<FAsyncTaskNotification>(NotificationConfig);
