Restoring functionality of SpawnAndRegisterTabWithId for the ability to spawn multiple instances of the same editor utility widget

diff --git a/Engine/Source/Editor/Blutility/Private/EditorUtilitySubsystem.cpp b/Engine/Source/Editor/Blutility/Private/EditorUtilitySubsystem.cpp
--- a/Engine/Source/Editor/Blutility/Private/EditorUtilitySubsystem.cpp
+++ b/Engine/Source/Editor/Blutility/Private/EditorUtilitySubsystem.cpp
@@ -247,7 +247,9 @@ UEditorUtilityWidget* UEditorUtilitySubsystem::SpawnAndRegisterTab(class UEditor
 
 UEditorUtilityWidget* UEditorUtilitySubsystem::SpawnAndRegisterTabWithId(class UEditorUtilityWidgetBlueprint* InBlueprint, FName InTabID)
 {
-	return SpawnAndRegisterTabAndGetID(InBlueprint, InTabID);
+	RegisterTabAndGetID(InBlueprint, InTabID);
+	SpawnRegisteredTabByID(InTabID);
+	return FindUtilityWidgetFromBlueprint(InBlueprint);
 }
 
 void UEditorUtilitySubsystem::RegisterTabAndGetID(class UEditorUtilityWidgetBlueprint* InBlueprint, FName& NewTabID)
@@ -348,15 +350,21 @@ bool UEditorUtilitySubsystem::SpawnRegisteredTabByID(FName NewTabID)
 		{
 			if (LevelEditorTabManager->HasTabSpawner(NewTabID))
 			{
-				TSharedPtr<SDockTab> NewDockTab = LevelEditorTabManager->TryInvokeTab(NewTabID);
-				NewDockTab->SetEnabled(FSlateApplication::Get().GetNormalExecutionAttribute());
-				IBlutilityModule* BlutilityModule = FModuleManager::GetModulePtr<IBlutilityModule>("Blutility");
-				UEditorUtilityWidgetBlueprint** WidgetToSpawn = RegisteredTabs.Find(NewTabID);
-				if (WidgetToSpawn)
+				if (TSharedPtr<SDockTab> NewDockTab = LevelEditorTabManager->TryInvokeTab(NewTabID))
+				{
+					NewDockTab->SetEnabled(FSlateApplication::Get().GetNormalExecutionAttribute());
+					IBlutilityModule* BlutilityModule = FModuleManager::GetModulePtr<IBlutilityModule>("Blutility");
+					UEditorUtilityWidgetBlueprint** WidgetToSpawn = RegisteredTabs.Find(NewTabID);
+					if (WidgetToSpawn)
+					{
+						check(*WidgetToSpawn);
+						BlutilityModule->AddLoadedScriptUI(*WidgetToSpawn);
+						return true;
+					}
+				}
+				else
 				{
-					check(*WidgetToSpawn);
-					BlutilityModule->AddLoadedScriptUI(*WidgetToSpawn);
-					return true;
+					UE_LOG(LogEditorUtilityBlueprint, Error, TEXT("TryInvokeTab failed with TabId: %s"), *NewTabID.ToString());
 				}
 			}
 		}
