#UE Perf improvement for opening the context menu in the content browser with a lot of assets

diff --git a/Engine/Source/Editor/Blutility/Private/BlutilityContentBrowserExtensions.cpp b/Engine/Source/Editor/Blutility/Private/BlutilityContentBrowserExtensions.cpp
--- a/Engine/Source/Editor/Blutility/Private/BlutilityContentBrowserExtensions.cpp
+++ b/Engine/Source/Editor/Blutility/Private/BlutilityContentBrowserExtensions.cpp
@@ -67,12 +67,13 @@ void FBlutilityContentBrowserExtensions::RegisterMenus()
 		// Run thru the assets to determine if any meet our criteria
 		TMap<TSharedRef<FAssetActionUtilityPrototype>, TSet<int32>> UtilityAndSelectionIndices;
 		TArray<FAssetData> SupportedAssets;
+		TMap<FAssetData, int32> ProcessedAssetIndices;
 		
 		if (SelectedAssets.Num() > 0)
 		{
 			FMessageLog EditorErrors("EditorErrors");
 			
-			auto ProcessAssetAction = [&EditorErrors, &SupportedAssets, &UtilityAndSelectionIndices, &SelectedAssets](const TSharedRef<FAssetActionUtilityPrototype>& ActionUtilityPrototype)
+			auto ProcessAssetAction = [&EditorErrors, &SupportedAssets, &ProcessedAssetIndices, &UtilityAndSelectionIndices, &SelectedAssets](const TSharedRef<FAssetActionUtilityPrototype>& ActionUtilityPrototype)
 			{
 				if (ActionUtilityPrototype->IsLatestVersion())
 				{
@@ -104,7 +105,13 @@ void FBlutilityContentBrowserExtensions::RegisterMenus()
 
 							if (bPassesClassFilter)
 							{
-								const int32 Index = SupportedAssets.AddUnique(Asset);
+								int32 Index = ProcessedAssetIndices.FindRef(Asset, INDEX_NONE);
+								if (Index == INDEX_NONE)
+								{
+									Index = SupportedAssets.Add(Asset);
+									ProcessedAssetIndices.Add(Asset, Index);
+								}
+								
 								UtilityAndSelectionIndices.FindOrAdd(ActionUtilityPrototype).Add(Index);
 							}
 						}
