Asset actions context menu optimizations
Change functor parameter type to const ref of soft class path to allow inner weak object ptr cache to be populated.
Also look up & cache asset classes outside of loop over asset actions.

diff --git a/Engine/Source/Editor/Blutility/Private/BlutilityContentBrowserExtensions.cpp b/Engine/Source/Editor/Blutility/Private/BlutilityContentBrowserExtensions.cpp
--- a/Engine/Source/Editor/Blutility/Private/BlutilityContentBrowserExtensions.cpp
+++ b/Engine/Source/Editor/Blutility/Private/BlutilityContentBrowserExtensions.cpp
@@ -3,6 +3,7 @@
 #include "BlutilityContentBrowserExtensions.h"
 
 #include "Algo/AnyOf.h"
+#include "Algo/Transform.h"
 #include "AssetActionUtility.h"
 #include "AssetRegistry/AssetData.h"
 #include "BlutilityMenuExtensions.h"
@@ -73,47 +74,54 @@ void FBlutilityContentBrowserExtensions::RegisterMenus()
 		{
 			FMessageLog EditorErrors("EditorErrors");
 			
-			auto ProcessAssetAction = [&EditorErrors, &SupportedAssets, &ProcessedAssetIndices, &UtilityAndSelectionIndices, &SelectedAssets](const TSharedRef<FAssetActionUtilityPrototype>& ActionUtilityPrototype)
+			// Resolve classes once to avoid doing it for each blutility
+			TArray<UClass*> SelectedAssetClasses;
+			Algo::Transform(SelectedAssets, SelectedAssetClasses, [](const FAssetData& Asset) { return Asset.GetClass(EResolveClass::Yes); });
+			auto ProcessAssetAction = [&EditorErrors, &SupportedAssets, &ProcessedAssetIndices, &UtilityAndSelectionIndices, &SelectedAssets, &SelectedAssetClasses](const TSharedRef<FAssetActionUtilityPrototype>& ActionUtilityPrototype)
 			{
 				if (ActionUtilityPrototype->IsLatestVersion())
 				{
 					TArray<TSoftClassPtr<UObject>> SupportedClassPtrs = ActionUtilityPrototype->GetSupportedClasses();
-					if (SupportedClassPtrs.Num() > 0)
+					if (SupportedClassPtrs.Num() == 0)
 					{
-						const bool bIsActionForBlueprints = ActionUtilityPrototype->AreSupportedClassesForBlueprints();
+						return;
+					}
+
+					const bool bIsActionForBlueprints = ActionUtilityPrototype->AreSupportedClassesForBlueprints();
 
-						for (const FAssetData& Asset : SelectedAssets)
+					for (int32 SelectedAssetIndex = 0; SelectedAssetIndex < SelectedAssets.Num(); ++SelectedAssetIndex)
+					{
+						const FAssetData& Asset = SelectedAssets[SelectedAssetIndex];
+						bool bPassesClassFilter = false;
+						if (bIsActionForBlueprints)
 						{
-							bool bPassesClassFilter = false;
-							if (bIsActionForBlueprints)
+							if (TSubclassOf<UBlueprint> AssetClass = SelectedAssetClasses[SelectedAssetIndex])
 							{
-								if (TSubclassOf<UBlueprint> AssetClass = Asset.GetClass())
+								if (const UClass* Blueprint_ParentClass = UBlueprint::GetBlueprintParentClassFromAssetTags(Asset))
 								{
-									if (const UClass* Blueprint_ParentClass = UBlueprint::GetBlueprintParentClassFromAssetTags(Asset))
-	                                {
-										bPassesClassFilter = 
-											Algo::AnyOf(SupportedClassPtrs, [Blueprint_ParentClass](TSoftClassPtr<UObject> ClassPtr){ return Blueprint_ParentClass->IsChildOf(ClassPtr.Get()); });
-	                                }
+									bPassesClassFilter = 
+										Algo::AnyOf(SupportedClassPtrs, [Blueprint_ParentClass](const TSoftClassPtr<UObject>& ClassPtr){ return Blueprint_ParentClass->IsChildOf(ClassPtr.Get()); });
 								}
 							}
-							else
-							{
-								// Is the asset the right kind?
-								bPassesClassFilter = 
-									Algo::AnyOf(SupportedClassPtrs, [&Asset](TSoftClassPtr<UObject> ClassPtr){ return Asset.IsInstanceOf(ClassPtr.Get(), EResolveClass::Yes); });
-							}
+						}
+						else
+						{
+							UClass* AssetClass = SelectedAssetClasses[SelectedAssetIndex];
+							// Is the asset the right kind?
+							bPassesClassFilter = AssetClass != nullptr &&
+								Algo::AnyOf(SupportedClassPtrs, [AssetClass](const TSoftClassPtr<UObject>& ClassPtr){ return AssetClass->IsChildOf(ClassPtr.Get()); });
+						}
 
-							if (bPassesClassFilter)
+						if (bPassesClassFilter)
+						{
+							int32 Index = ProcessedAssetIndices.FindRef(Asset, INDEX_NONE);
+							if (Index == INDEX_NONE)
 							{
-								int32 Index = ProcessedAssetIndices.FindRef(Asset, INDEX_NONE);
-								if (Index == INDEX_NONE)
-								{
-									Index = SupportedAssets.Add(Asset);
-									ProcessedAssetIndices.Add(Asset, Index);
-								}
-								
-								UtilityAndSelectionIndices.FindOrAdd(ActionUtilityPrototype).Add(Index);
+								Index = SupportedAssets.Add(Asset);
+								ProcessedAssetIndices.Add(Asset, Index);
 							}
+							
+							UtilityAndSelectionIndices.FindOrAdd(ActionUtilityPrototype).Add(Index);
 						}
 					}
 				}
