Unshelved from pending changelist '34134820':

Explictly map EditorFunctionLibrary to EditorUtilityBlueprint, this allows us to maintain this invariant when assets are created from scripts, rather than custom UI
#rnx

diff --git a/Engine/Source/Editor/Blutility/Private/BlutilityModule.cpp b/Engine/Source/Editor/Blutility/Private/BlutilityModule.cpp
--- a/Engine/Source/Editor/Blutility/Private/BlutilityModule.cpp
+++ b/Engine/Source/Editor/Blutility/Private/BlutilityModule.cpp
@@ -13,6 +13,7 @@
 #include "Delegates/Delegate.h"
 #include "Editor.h"
 #include "Editor/EditorEngine.h"
+#include "EditorFunctionLibrary.h"
 #include "EditorSupportDelegates.h"
 #include "EditorUtilityActor.h"
 #include "EditorUtilityBlueprint.h"
@@ -96,6 +97,7 @@ public:
 		KismetCompilerModule.OverrideBPTypeForClass(AEditorUtilityActor::StaticClass(), UEditorUtilityBlueprint::StaticClass());
 		KismetCompilerModule.OverrideBPTypeForClass(AEditorUtilityCamera::StaticClass(), UEditorUtilityBlueprint::StaticClass());
 		KismetCompilerModule.OverrideBPTypeForClass(UEditorUtilityObject::StaticClass(), UEditorUtilityBlueprint::StaticClass());
+		KismetCompilerModule.OverrideBPTypeForClass(UEditorFunctionLibrary::StaticClass(), UEditorUtilityBlueprint::StaticClass());
 		KismetCompilerModule.OverrideBPTypeForClass(UEditorUtilityWidget::StaticClass(), UEditorUtilityWidgetBlueprint::StaticClass());
 
 		FBlutilityContentBrowserExtensions::InstallHooks();
