Fix to EUW Susbsystem console command registration

diff --git a/Engine/Source/Editor/Blutility/Private/EditorUtilitySubsystem.cpp b/Engine/Source/Editor/Blutility/Private/EditorUtilitySubsystem.cpp
--- a/Engine/Source/Editor/Blutility/Private/EditorUtilitySubsystem.cpp
+++ b/Engine/Source/Editor/Blutility/Private/EditorUtilitySubsystem.cpp
@@ -56,19 +56,26 @@ UEditorUtilitySubsystem::UEditorUtilitySubsystem()
 
 void UEditorUtilitySubsystem::Initialize(FSubsystemCollectionBase& Collection)
 {
-	RunTaskCommandObject = IConsoleManager::Get().RegisterConsoleCommand(
-		TEXT("RunTask"),
-		TEXT(""),
-		FConsoleCommandWithWorldArgsAndOutputDeviceDelegate::CreateUObject(this, &UEditorUtilitySubsystem::RunTaskCommand),
-		ECVF_Default
-	);
-
-	CancelAllTasksCommandObject = IConsoleManager::Get().RegisterConsoleCommand(
-		TEXT("CancelAllTasks"),
-		TEXT(""),
-		FConsoleCommandWithWorldArgsAndOutputDeviceDelegate::CreateUObject(this, &UEditorUtilitySubsystem::CancelAllTasksCommand),
-		ECVF_Default
-	);
+
+	if (!IConsoleManager::Get().FindConsoleObject(TEXT("RunTask")))
+	{
+		RunTaskCommandObject = IConsoleManager::Get().RegisterConsoleCommand(
+			TEXT("RunTask"),
+			TEXT(""),
+			FConsoleCommandWithWorldArgsAndOutputDeviceDelegate::CreateUObject(this, &UEditorUtilitySubsystem::RunTaskCommand),
+			ECVF_Default
+		);
+	}
+
+	if (!IConsoleManager::Get().FindConsoleObject(TEXT("CancelAllTasks")))
+	{
+		CancelAllTasksCommandObject = IConsoleManager::Get().RegisterConsoleCommand(
+			TEXT("CancelAllTasks"),
+			TEXT(""),
+			FConsoleCommandWithWorldArgsAndOutputDeviceDelegate::CreateUObject(this, &UEditorUtilitySubsystem::CancelAllTasksCommand),
+			ECVF_Default
+		);
+	}
 
 	IMainFrameModule& MainFrameModule = IMainFrameModule::Get();
 	if (MainFrameModule.IsWindowInitialized())
@@ -98,7 +105,15 @@ void UEditorUtilitySubsystem::Deinitialize()
 
 	FTSTicker::GetCoreTicker().RemoveTicker(TickerHandle);
 
-	IConsoleManager::Get().UnregisterConsoleObject(RunTaskCommandObject);
+	if (RunTaskCommandObject)
+	{
+		IConsoleManager::Get().UnregisterConsoleObject(RunTaskCommandObject); 
+	}
+
+	if (CancelAllTasksCommandObject)
+	{
+		IConsoleManager::Get().UnregisterConsoleObject(CancelAllTasksCommandObject);
+	}
 
 	FEditorDelegates::BeginPIE.RemoveAll(this);
 	FEditorDelegates::EndPIE.RemoveAll(this);
