Editor Utility Widget blueprint
-Added an option to spawn an Editor Utilty Widget as a nomad tab so it can be docked as both major and minor tab.

diff --git a/Engine/Source/Editor/Blutility/Classes/EditorUtilityWidgetBlueprint.h b/Engine/Source/Editor/Blutility/Classes/EditorUtilityWidgetBlueprint.h
--- a/Engine/Source/Editor/Blutility/Classes/EditorUtilityWidgetBlueprint.h
+++ b/Engine/Source/Editor/Blutility/Classes/EditorUtilityWidgetBlueprint.h
@@ -94,6 +94,10 @@ private:
 	UPROPERTY(Category = Settings, EditDefaultsOnly)
 	bool bIsEnabledInDebugging = false;
 
+	// Should the widget be spawned on a Nomad tab to be docked anywhere
+	UPROPERTY(Category = Settings, EditDefaultsOnly)
+	bool bSpawnAsNomadTab = false;
+
 	FName RegistrationName;
 
 	TWeakPtr<SDockTab> CreatedTab;
diff --git a/Engine/Source/Editor/Blutility/Private/EditorUtilityWidgetBlueprint.cpp b/Engine/Source/Editor/Blutility/Private/EditorUtilityWidgetBlueprint.cpp
--- a/Engine/Source/Editor/Blutility/Private/EditorUtilityWidgetBlueprint.cpp
+++ b/Engine/Source/Editor/Blutility/Private/EditorUtilityWidgetBlueprint.cpp
@@ -66,6 +66,24 @@ void UEditorUtilityWidgetBlueprint::BeginDestroy()
 
 TSharedRef<SDockTab> UEditorUtilityWidgetBlueprint::SpawnEditorUITab(const FSpawnTabArgs& SpawnTabArgs)
 {
+	if(bSpawnAsNomadTab)
+	{
+		TSharedRef<SDockTab> SpawnedTab = SNew(SDockTab).TabRole(NomadTab);
+
+		TSharedRef<SWidget> TabWidget = CreateUtilityWidget();
+		SpawnedTab->SetContent(TabWidget);
+	
+		SpawnedTab->SetOnTabClosed(SDockTab::FOnTabClosedCallback::CreateUObject(this, &UEditorUtilityWidgetBlueprint::UpdateRespawnListIfNeeded));
+		CreatedTab = SpawnedTab;
+	
+		OnCompiled().AddUObject(this, &UEditorUtilityWidgetBlueprint::RegenerateCreatedTab);
+	
+		FLevelEditorModule& LevelEditor = FModuleManager::LoadModuleChecked<FLevelEditorModule>("LevelEditor");
+		LevelEditor.OnMapChanged().AddUObject(this, &UEditorUtilityWidgetBlueprint::ChangeTabWorld);
+
+		return SpawnedTab;
+	}
+		
 	TSharedRef<SDockTab> SpawnedTab = SNew(SDockTab);
 
 	TSharedRef<SWidget> TabWidget = CreateUtilityWidget();
