Adding a ToolMenu widget so Editor Utility Widgets can create their own ToolMenus and maintain editor styling without needing to recreate anything from scratch.

diff --git a/Engine/Source/Editor/Blutility/Private/ToolMenuWidget.cpp b/Engine/Source/Editor/Blutility/Private/ToolMenuWidget.cpp
new file mode 100644
--- /dev/null
+++ b/Engine/Source/Editor/Blutility/Private/ToolMenuWidget.cpp
@@ -0,0 +1,90 @@
+// Copyright Epic Games, Inc. All Rights Reserved.
+
+
+#include "ToolMenuWidget.h"
+#include "ToolMenus.h"
+#include "Widgets/Layout/SSpacer.h"
+#include "EditorUtilityWidgetBlueprint.h"
+#include "EditorUtilityWidget.h"
+#include "UObject/Package.h"
+
+
+
+UToolMenuWidget::UToolMenuWidget(const FObjectInitializer& ObjectInitializer)
+	: Super(ObjectInitializer)
+	, MenuName("MenuBar")
+	, MenuType(EMultiBoxType::MenuBar)
+{
+	if (MyToolMenu == SNullWidget::NullWidget)
+	{
+		MyToolMenu = SNew(SSpacer);
+	}
+	UpdateFullMenuName();
+}
+
+TSharedRef<SWidget> UToolMenuWidget::RebuildWidget()
+{
+	if (!UToolMenus::Get()->IsMenuRegistered(FullMenuName))
+	{
+		UToolMenu* ToolBar = UToolMenus::Get()->RegisterMenu(FullMenuName, NAME_None, MenuType);
+	}
+
+	FToolMenuContext MenuContext;
+	MyToolMenu = UToolMenus::Get()->GenerateWidget(FullMenuName, MenuContext);
+	return MyToolMenu;
+}
+
+void UToolMenuWidget::UpdateFullMenuName()
+{
+	if (HasAnyFlags(RF_ClassDefaultObject))
+	{
+		return;
+	}
+
+	// PostEditChangeProperty is called on the transient and the class object.
+	// To get a consistent name representing the containing editor utility, either the EditorUtilityWidgetBlueprint or the EditorUtilityWidget needs to be used.
+	UEditorUtilityWidgetBlueprint* EUWBlueprintOuter = GetTypedOuter<UEditorUtilityWidgetBlueprint>();
+	UEditorUtilityWidget* EUWOuter = GetTypedOuter<UEditorUtilityWidget>();
+
+	if (EUWBlueprintOuter)
+	{
+		FullMenuName = FName(*(FString::Format(TEXT("{0}.{1}"), { *(EUWBlueprintOuter->GetName()), *MenuName })));
+	}
+	else if (EUWOuter)
+	{
+		FString PackageName = EUWOuter->GetArchetype()->GetPackage()->GetName();
+		FullMenuName = FName(*(FString::Format(TEXT("{0}.{1}"), { *(FPaths::GetBaseFilename(PackageName)), *MenuName })));
+	}
+}
+
+#if WITH_EDITOR
+
+void UToolMenuWidget::PostInitProperties()
+{
+	Super::PostInitProperties();
+
+	if (IsValid(this))
+	{
+		UpdateFullMenuName();
+	}
+}
+
+void UToolMenuWidget::PostEditChangeProperty(struct FPropertyChangedEvent& PropertyChangedEvent)
+{
+	Super::PostEditChangeProperty(PropertyChangedEvent);
+
+	if (PropertyChangedEvent.Property)
+	{
+		static const FName MenuNameProp("MenuName");
+
+		FName PropertyName = PropertyChangedEvent.Property->GetFName();
+
+		if (PropertyName == MenuNameProp)
+		{
+			UpdateFullMenuName();
+		}
+	}
+	
+}
+
+#endif //WITH_EDITOR
diff --git a/Engine/Source/Editor/Blutility/Public/ToolMenuWidget.h b/Engine/Source/Editor/Blutility/Public/ToolMenuWidget.h
new file mode 100644
--- /dev/null
+++ b/Engine/Source/Editor/Blutility/Public/ToolMenuWidget.h
@@ -0,0 +1,64 @@
+// Copyright Epic Games, Inc. All Rights Reserved.
+
+#pragma once
+
+#include "CoreMinimal.h"
+#include "Components/Widget.h"
+#include "Widgets/SNullWidget.h"
+
+#include "ToolMenuWidget.generated.h"
+
+enum class EMultiBoxType : uint8;
+
+/**
+ * A base ToolMenu widget that can support Menus and Toolbars.
+ * 
+ * Menus can be modified and added to with Blueprint or Python commands using the ToolMenus system.
+ * 
+ */
+UCLASS(MinimalAPI)
+class UToolMenuWidget : public UWidget
+{
+	GENERATED_BODY()
+
+	UToolMenuWidget(const FObjectInitializer& ObjectInitializer);
+
+public:
+	/** Name of toolbar for registering with ToolMenus.
+	*	This has the editor utility widget pre-pended to it to make the FullMenuName.
+	*/
+	UPROPERTY(Category = Config, EditAnywhere, BlueprintReadOnly)
+	FString MenuName;
+
+	UPROPERTY(Category = Config, EditAnywhere, BlueprintReadOnly)
+	EMultiBoxType MenuType;
+
+public:
+
+#if WITH_EDITOR
+	//~ Begin UObject Interface
+	BLUTILITY_API virtual void PostInitProperties();
+	//~ End UObject Interface
+
+	//~ Begin UObject Interface
+	BLUTILITY_API virtual void PostEditChangeProperty(struct FPropertyChangedEvent& PropertyChangedEvent) override;
+	//~ End UObject Interface
+#endif //WITH_EDITOR
+
+protected:
+	TSharedRef<SWidget> MyToolMenu = SNullWidget::NullWidget;
+	
+	/** The final usable name that can be retrieved by Python/Blueprint.
+	*	Use this value when trying to Extend menus.
+	*/
+	UPROPERTY(Category = Config, BlueprintReadOnly, VisibleAnywhere)
+	FName FullMenuName;
+
+protected:
+	//~ Begin UWidget Interface
+	BLUTILITY_API virtual TSharedRef<SWidget> RebuildWidget() override;
+	//~ End UWidget Interface
+	
+private:
+	void UpdateFullMenuName();
+};
