Fix IsEditorUtilityBlueprint to return true for any class that is editor only

[RN] Any blueprint that is based on an editor only class can now use the editor utility APIs

diff --git a/Engine/Source/Editor/Blutility/Private/BlutilityModule.cpp b/Engine/Source/Editor/Blutility/Private/BlutilityModule.cpp
--- a/Engine/Source/Editor/Blutility/Private/BlutilityModule.cpp
+++ b/Engine/Source/Editor/Blutility/Private/BlutilityModule.cpp
@@ -35,6 +35,7 @@
 #include "IBlutilityModule.h"
 #include "Internationalization/Internationalization.h"
 #include "Internationalization/Text.h"
+#include "Kismet2/BlueprintEditorUtils.h"
 #include "KismetCompiler.h"
 #include "KismetCompilerModule.h"
 #include "LevelEditor.h"
@@ -217,14 +218,34 @@ public:
 	virtual bool IsEditorUtilityBlueprint( const UBlueprint* Blueprint ) const override
 	{
 		const UClass* BPClass = Blueprint ? Blueprint->GetClass() : nullptr;
-
-		if( BPClass && 
+		const bool bIsExplicitEditorUtility = BPClass && 
 			( BPClass->IsChildOf( UEditorUtilityBlueprint::StaticClass() )
 				|| BPClass->IsChildOf( UEditorUtilityWidgetBlueprint::StaticClass() )
-				|| IsEditorOnlyObject(BPClass) ) )
+				|| IsEditorOnlyObject(BPClass) );
+
+		if( bIsExplicitEditorUtility )
 		{
 			return true;
 		}
+		else
+		{
+			// People are frequently defining native editor only UObjects and 
+			// using the editor only APIs. For consistency, indicate to callers
+			// that these editor only classes should behave as editor utilities,
+			// unless they are actors that want to run during pie:
+			const UClass* NativeParent = FBlueprintEditorUtils::GetNativeParent(Blueprint);
+			if(NativeParent && IsEditorOnlyObject(NativeParent))
+			{
+				if (NativeParent->IsChildOf(AActor::StaticClass()))
+				{
+					if(const AActor* NativeCDO = Cast<AActor>(NativeParent->ClassDefaultObject))
+					{
+						return !NativeCDO->IsNonPIEEditorOnly();
+					}
+				}
+				return true;
+			}
+		}
 		return false;
 	}
 
