EditorUtilitySubsystem : When bSpawnAsNomadTab is set to true for an UEditorUtilityWidgetBlueprint then its tab spawner needs to be registered with FGlobalTabmanager::Get()->RegisterNomadTabSpawner.

TabManager : In AttemptToOpenTab remove the tab from the collapsed area if it is a nomad tab in a floating window, and the tab manager is _not_ the global tab manager.

This fixes a bug where EUWs with this flag set did not fully function as nomad tabs.  The TabManager change is required to cleanup when changing this flag from false to true between runs of the EUW from its asset editor.

diff --git a/Engine/Source/Editor/Blutility/Classes/EditorUtilityWidgetBlueprint.h b/Engine/Source/Editor/Blutility/Classes/EditorUtilityWidgetBlueprint.h
--- a/Engine/Source/Editor/Blutility/Classes/EditorUtilityWidgetBlueprint.h
+++ b/Engine/Source/Editor/Blutility/Classes/EditorUtilityWidgetBlueprint.h
@@ -74,6 +74,11 @@ public:
 	/** Returns the default desired tab display name that was specified for this widget */
 	FText GetTabDisplayName() const;
 
+	bool ShouldSpawnAsNomadTab() const
+	{
+		return bSpawnAsNomadTab;
+	}
+
 	virtual UWidgetEditingProjectSettings* GetRelevantSettings() override;
 	virtual const UWidgetEditingProjectSettings* GetRelevantSettings() const override;
 
diff --git a/Engine/Source/Editor/Blutility/Private/EditorUtilitySubsystem.cpp b/Engine/Source/Editor/Blutility/Private/EditorUtilitySubsystem.cpp
--- a/Engine/Source/Editor/Blutility/Private/EditorUtilitySubsystem.cpp
+++ b/Engine/Source/Editor/Blutility/Private/EditorUtilitySubsystem.cpp
@@ -275,17 +275,43 @@ void UEditorUtilitySubsystem::RegisterTabAndGetID(class UEditorUtilityWidgetBlue
 	{
 		FName RegistrationName = NewTabID.IsNone() ? FName(*(InBlueprint->GetPathName() + LOCTEXT("ActiveTabSuffix", "_ActiveTab").ToString())) : FName(*(InBlueprint->GetPathName() + NewTabID.ToString()));
 		FText DisplayName = InBlueprint->GetTabDisplayName();
+		IBlutilityModule& BlutilityModule = FModuleManager::GetModuleChecked<IBlutilityModule>("Blutility");
 		FLevelEditorModule& LevelEditorModule = FModuleManager::GetModuleChecked<FLevelEditorModule>(TEXT("LevelEditor"));
+
 		if (TSharedPtr<FTabManager> LevelEditorTabManager = LevelEditorModule.GetLevelEditorTabManager())
 		{
-			if (!LevelEditorTabManager->HasTabSpawner(RegistrationName))
+			// Nomad tabs need to have their tab spawner registered with the global tab manager via RegisterNomadTabSpawner
+			if (InBlueprint->ShouldSpawnAsNomadTab())
 			{
-				IBlutilityModule* BlutilityModule = FModuleManager::GetModulePtr<IBlutilityModule>("Blutility");
-				LevelEditorTabManager->RegisterTabSpawner(RegistrationName, FOnSpawnTab::CreateUObject(InBlueprint, &UEditorUtilityWidgetBlueprint::SpawnEditorUITab))
-					.SetDisplayName(DisplayName)
-					.SetGroup(BlutilityModule->GetMenuGroup().ToSharedRef());
-				InBlueprint->SetRegistrationName(RegistrationName);
+				if (LevelEditorTabManager->HasTabSpawner(RegistrationName))
+				{
+					LevelEditorTabManager->UnregisterTabSpawner(RegistrationName);
+				}
+
+				if (!FGlobalTabmanager::Get()->HasTabSpawner(RegistrationName))
+				{
+					FGlobalTabmanager::Get()->RegisterNomadTabSpawner(RegistrationName, FOnSpawnTab::CreateUObject(InBlueprint, &UEditorUtilityWidgetBlueprint::SpawnEditorUITab))
+						.SetDisplayName(DisplayName)
+						.SetGroup(BlutilityModule.GetMenuGroup().ToSharedRef());
+					InBlueprint->SetRegistrationName(RegistrationName);
+				}
 			}
+			else
+			{
+				if (FGlobalTabmanager::Get()->HasTabSpawner(RegistrationName))
+				{
+					FGlobalTabmanager::Get()->UnregisterNomadTabSpawner(RegistrationName);
+				}
+
+				if (!LevelEditorTabManager->HasTabSpawner(RegistrationName))
+				{
+					LevelEditorTabManager->RegisterTabSpawner(RegistrationName, FOnSpawnTab::CreateUObject(InBlueprint, &UEditorUtilityWidgetBlueprint::SpawnEditorUITab))
+						.SetDisplayName(DisplayName)
+						.SetGroup(BlutilityModule.GetMenuGroup().ToSharedRef());
+					InBlueprint->SetRegistrationName(RegistrationName);
+				}
+			}
+
 			RegisteredTabs.Add(RegistrationName, InBlueprint);
 			NewTabID = RegistrationName;
 		}
@@ -365,25 +391,24 @@ bool UEditorUtilitySubsystem::SpawnRegisteredTabByID(FName NewTabID)
 		FLevelEditorModule& LevelEditorModule = FModuleManager::GetModuleChecked<FLevelEditorModule>(TEXT("LevelEditor"));
 		if (TSharedPtr<FTabManager> LevelEditorTabManager = LevelEditorModule.GetLevelEditorTabManager())
 		{
-			if (LevelEditorTabManager->HasTabSpawner(NewTabID))
+			// Invoke the tab via the LevelEditor's TabManager even if the tab spawner is registered globally as a nomad tab.
+			// TryInvokeTab will invoke in the global tab manager if the tab is not found in the LevelEditor.
+			if (TSharedPtr<SDockTab> NewDockTab = LevelEditorTabManager->TryInvokeTab(NewTabID))
 			{
-				if (TSharedPtr<SDockTab> NewDockTab = LevelEditorTabManager->TryInvokeTab(NewTabID))
-				{
-					NewDockTab->SetEnabled(FSlateApplication::Get().GetNormalExecutionAttribute());
-					IBlutilityModule* BlutilityModule = FModuleManager::GetModulePtr<IBlutilityModule>("Blutility");
-					UEditorUtilityWidgetBlueprint** WidgetToSpawn = RegisteredTabs.Find(NewTabID);
-					if (WidgetToSpawn)
-					{
-						check(*WidgetToSpawn);
-						BlutilityModule->AddLoadedScriptUI(*WidgetToSpawn);
-						return true;
-					}
-				}
-				else
+				NewDockTab->SetEnabled(FSlateApplication::Get().GetNormalExecutionAttribute());
+				IBlutilityModule* BlutilityModule = FModuleManager::GetModulePtr<IBlutilityModule>("Blutility");
+				UEditorUtilityWidgetBlueprint** WidgetToSpawn = RegisteredTabs.Find(NewTabID);
+				if (WidgetToSpawn)
 				{
-					UE_LOG(LogEditorUtilityBlueprint, Error, TEXT("TryInvokeTab failed with TabId: %s"), *NewTabID.ToString());
+					check(*WidgetToSpawn);
+					BlutilityModule->AddLoadedScriptUI(*WidgetToSpawn);
+					return true;
 				}
 			}
+			else
+			{
+				UE_LOG(LogEditorUtilityBlueprint, Error, TEXT("TryInvokeTab failed with TabId: %s"), *NewTabID.ToString());
+			}
 		}
 	}
 	return false;
@@ -402,6 +427,11 @@ bool UEditorUtilitySubsystem::DoesTabExist(FName NewTabID)
 				return true;
 			}
 		}
+
+		if (TSharedPtr<SDockTab> FoundNomadTab = FGlobalTabmanager::Get()->FindExistingLiveTab(NewTabID))
+		{
+			return true;
+		}
 	}
 
 	return false;
@@ -421,6 +451,12 @@ bool UEditorUtilitySubsystem::CloseTabByID(FName NewTabID)
 				return true;
 			}
 		}
+
+		if (TSharedPtr<SDockTab> FoundNomadTab = FGlobalTabmanager::Get()->FindExistingLiveTab(NewTabID))
+		{
+			FoundNomadTab->RequestCloseTab();
+			return true;
+		}
 	}
 
 	return false;
@@ -432,16 +468,32 @@ bool UEditorUtilitySubsystem::UnregisterTabByID(FName TabID)
 	{
 		const FLevelEditorModule& LevelEditorModule = FModuleManager::GetModuleChecked<FLevelEditorModule>(TEXT("LevelEditor"));
 		const TSharedPtr<FTabManager> LevelEditorTabManager = LevelEditorModule.GetLevelEditorTabManager();
-		const TSharedPtr<SDockTab> FoundTab = LevelEditorTabManager->FindExistingLiveTab(TabID);
+		TSharedPtr<SDockTab> FoundTab = LevelEditorTabManager->FindExistingLiveTab(TabID);
+
+		if (!FoundTab)
+		{
+			FoundTab = FGlobalTabmanager::Get()->FindExistingLiveTab(TabID);
+		}
+
 		if (FoundTab)
 		{
 			FoundTab->RequestCloseTab();
 		}
 
-		if (RegisteredTabs.Contains(TabID) && LevelEditorTabManager->HasTabSpawner(TabID))
+		if (RegisteredTabs.Contains(TabID))
 		{
 			RegisteredTabs.Remove(TabID);
-			LevelEditorTabManager->UnregisterTabSpawner(TabID);
+			
+			if (LevelEditorTabManager->HasTabSpawner(TabID))
+			{
+				LevelEditorTabManager->UnregisterTabSpawner(TabID);
+			}
+
+			if (FGlobalTabmanager::Get()->HasTabSpawner(TabID))
+			{
+				FGlobalTabmanager::Get()->UnregisterNomadTabSpawner(TabID);
+			}
+
 			return true;
 		}
 	}
diff --git a/Engine/Source/Runtime/Slate/Private/Framework/Docking/TabManager.cpp b/Engine/Source/Runtime/Slate/Private/Framework/Docking/TabManager.cpp
--- a/Engine/Source/Runtime/Slate/Private/Framework/Docking/TabManager.cpp
+++ b/Engine/Source/Runtime/Slate/Private/Framework/Docking/TabManager.cpp
@@ -1506,16 +1506,26 @@ TSharedPtr<SDockingTabStack> FTabManager::AttemptToOpenTab( const FTabId& Closed
 	if ( CollapsedAreaWithMatchingTabIndex != INDEX_NONE )
 	{
 		TSharedRef<FTabManager::FArea> CollapsedAreaWithMatchingTab = CollapsedDockAreas[CollapsedAreaWithMatchingTabIndex];
-		
-		TSharedPtr<SDockingArea> RestoredArea = RestoreArea(CollapsedDockAreas[CollapsedAreaWithMatchingTabIndex],
-			GetPrivateApi().GetParentWindow(), false, EOutputCanBeNullptr::Never, bForceOpenWindowIfNeeded);
-		check(RestoredArea.IsValid());
-		// We have just un-collapsed this dock area.
-		// Don't rely on the collapsed tab index: RestoreArea() can end up kicking the task graph which could do other tab work and modify the CollapsedDockAreas array.
-		CollapsedDockAreas.Remove(CollapsedAreaWithMatchingTab);
-		if (RestoredArea.IsValid())
-		{
-			StackWithClosedTab = FindTabInLiveArea(TabMatcher, StaticCastSharedRef<SDockingArea>(RestoredArea->AsShared()));
+
+		// If this is not the global tab manager and the tab is a NomadTab in a floating window, then remove it from the collapsed area.
+		if (FGlobalTabmanager::Get() != SharedThis(this) &&
+			NomadTabSpawner->Contains(ClosedTabId.TabType) &&
+			CollapsedAreaWithMatchingTab->WindowPlacement != FTabManager::FArea::EWindowPlacement::Placement_NoWindow)
+		{
+			RemoveTabFromCollapsedAreas(TabMatcher);
+		}
+		else
+		{
+			TSharedPtr<SDockingArea> RestoredArea = RestoreArea(CollapsedDockAreas[CollapsedAreaWithMatchingTabIndex],
+				GetPrivateApi().GetParentWindow(), false, EOutputCanBeNullptr::Never, bForceOpenWindowIfNeeded);
+			check(RestoredArea.IsValid());
+			// We have just un-collapsed this dock area.
+			// Don't rely on the collapsed tab index: RestoreArea() can end up kicking the task graph which could do other tab work and modify the CollapsedDockAreas array.
+			CollapsedDockAreas.Remove(CollapsedAreaWithMatchingTab);
+			if (RestoredArea.IsValid())
+			{
+				StackWithClosedTab = FindTabInLiveArea(TabMatcher, StaticCastSharedRef<SDockingArea>(RestoredArea->AsShared()));
+			}
 		}
 	}
 
