[Scripted Asset Actions] Fix an issue where types derived from blueprints would not filter correctly in scripted asset actions for blueprint actions

This was caused by the modification to stop fully loading assets on right click. We were only walking the native class chain in the case the asset action was for blueprint types, since those are no longer loaded fully on a right click

Cvar to enable: Editor.AssetActions.ExperimentalClassSupport 1

diff --git a/Engine/Source/Editor/Blutility/Private/BlutilityContentBrowserExtensions.cpp b/Engine/Source/Editor/Blutility/Private/BlutilityContentBrowserExtensions.cpp
--- a/Engine/Source/Editor/Blutility/Private/BlutilityContentBrowserExtensions.cpp
+++ b/Engine/Source/Editor/Blutility/Private/BlutilityContentBrowserExtensions.cpp
@@ -6,6 +6,7 @@
 #include "Algo/Transform.h"
 #include "AssetActionUtility.h"
 #include "AssetRegistry/AssetData.h"
+#include "AssetRegistry/IAssetRegistry.h"
 #include "BlutilityMenuExtensions.h"
 #include "Containers/Array.h"
 #include "Containers/Map.h"
@@ -37,6 +38,16 @@
 
 #define LOCTEXT_NAMESPACE "BlutilityContentBrowserExtensions"
 
+namespace UE::Editor::Utilities
+{
+	bool bEnableExperimentalScriptedActionsClassSupport = false;
+	FAutoConsoleVariableRef GCvarExpiermentalScriptedActionsClassSupport(
+		TEXT("Editor.AssetActions.ExperimentalClassSupport"),
+		bEnableExperimentalScriptedActionsClassSupport,
+		TEXT("Enables experimental support for iterating the class hierarchy of blueprints for scripted asset actions."),
+		FConsoleVariableDelegate());
+}
+
 void FBlutilityContentBrowserExtensions::InstallHooks()
 {
 	UToolMenus::RegisterStartupCallback(
@@ -97,10 +108,38 @@ void FBlutilityContentBrowserExtensions::RegisterMenus()
 						{
 							if (TSubclassOf<UBlueprint> AssetClass = SelectedAssetClasses[SelectedAssetIndex])
 							{
-								if (const UClass* Blueprint_ParentClass = UBlueprint::GetBlueprintParentClassFromAssetTags(Asset))
+								if (UE::Editor::Utilities::bEnableExperimentalScriptedActionsClassSupport)
+								{
+									FString ClassPathAsString;
+									Asset.GetTagValue(FBlueprintTags::GeneratedClassPath, ClassPathAsString);
+									FTopLevelAssetPath AssetClassPath(ClassPathAsString);
+
+									const IAssetRegistry& AssetRegistry = IAssetRegistry::GetChecked();
+									TArray<FTopLevelAssetPath> AssetClassPathHierarchy;
+									AssetRegistry.GetAncestorClassNames(AssetClassPath, AssetClassPathHierarchy);
+									AssetClassPathHierarchy.Insert(AssetClassPath, 0);
+
+									for (const FTopLevelAssetPath& ClassPath : AssetClassPathHierarchy)
+									{
+										bPassesClassFilter = Algo::AnyOf(SupportedClassPtrs,
+											[ClassPath](const TSoftClassPtr<UObject>& ClassPtr)
+											{
+												return ClassPtr.ToSoftObjectPath() == FSoftObjectPath(ClassPath);
+											});
+
+										if (bPassesClassFilter)
+										{
+											break;
+										}
+									}
+								}
+								else
 								{
-									bPassesClassFilter = 
-										Algo::AnyOf(SupportedClassPtrs, [Blueprint_ParentClass](const TSoftClassPtr<UObject>& ClassPtr){ return Blueprint_ParentClass->IsChildOf(ClassPtr.Get()); });
+									if (const UClass* Blueprint_ParentClass = UBlueprint::GetBlueprintParentClassFromAssetTags(Asset))
+									{
+										bPassesClassFilter = 
+											Algo::AnyOf(SupportedClassPtrs, [Blueprint_ParentClass](const TSoftClassPtr<UObject>& ClassPtr){ return Blueprint_ParentClass->IsChildOf(ClassPtr.Get()); });
+									}
 								}
 							}
 						}
