UMGEditor: Fix not generating variable GUIDs for default RootWidgets of WidgetBPs and EditorUtilityWidgetBPs

#rnx

diff --git a/Engine/Source/Editor/Blutility/Private/EditorUtilityWidgetBlueprintFactory.cpp b/Engine/Source/Editor/Blutility/Private/EditorUtilityWidgetBlueprintFactory.cpp
--- a/Engine/Source/Editor/Blutility/Private/EditorUtilityWidgetBlueprintFactory.cpp
+++ b/Engine/Source/Editor/Blutility/Private/EditorUtilityWidgetBlueprintFactory.cpp
@@ -146,6 +146,7 @@ UObject* UEditorUtilityWidgetBlueprintFactory::FactoryCreateNew(UClass* Class, U
 			{
 				UWidget* Root = NewBP->WidgetTree->ConstructWidget<UWidget>(RootWidgetPanel);
 				NewBP->WidgetTree->RootWidget = Root;
+				NewBP->OnVariableAdded(Root->GetFName());
 			}
 		}
 
diff --git a/Engine/Source/Editor/UMGEditor/Private/WidgetBlueprintFactory.cpp b/Engine/Source/Editor/UMGEditor/Private/WidgetBlueprintFactory.cpp
--- a/Engine/Source/Editor/UMGEditor/Private/WidgetBlueprintFactory.cpp
+++ b/Engine/Source/Editor/UMGEditor/Private/WidgetBlueprintFactory.cpp
@@ -193,6 +193,7 @@ UObject* UWidgetBlueprintFactory::FactoryCreateNew(UClass* Class, UObject* InPar
 			{
 				UWidget* Root = NewBP->WidgetTree->ConstructWidget<UWidget>(RootWidgetPanel);
 				NewBP->WidgetTree->RootWidget = Root;
+				NewBP->OnVariableAdded(Root->GetFName());
 			}
 		}
 
