Blutility: Add BlueprintCallable FindChildWidgetName() function to Editor Utility Widget

#rnx

diff --git a/Engine/Source/Editor/Blutility/Classes/EditorUtilityWidget.h b/Engine/Source/Editor/Blutility/Classes/EditorUtilityWidget.h
--- a/Engine/Source/Editor/Blutility/Classes/EditorUtilityWidget.h
+++ b/Engine/Source/Editor/Blutility/Classes/EditorUtilityWidget.h
@@ -30,6 +30,9 @@ public:
 	UFUNCTION(BlueprintCallable, BlueprintImplementableEvent, Category = "Editor")
 	void Run();
 
+	UFUNCTION(BlueprintCallable, BlueprintPure = false, Category = "Editor Utility Widget")
+	UWidget* FindChildWidgetByName(FName WidgetName) const;
+
 	// Run the default action
 	void ExecuteDefaultAction();
 
diff --git a/Engine/Source/Editor/Blutility/Private/EditorUtilityWidget.cpp b/Engine/Source/Editor/Blutility/Private/EditorUtilityWidget.cpp
--- a/Engine/Source/Editor/Blutility/Private/EditorUtilityWidget.cpp
+++ b/Engine/Source/Editor/Blutility/Private/EditorUtilityWidget.cpp
@@ -2,6 +2,7 @@
 
 #include "EditorUtilityWidget.h"
 
+#include "Blueprint/WidgetTree.h"
 #include "Internationalization/Internationalization.h"
 #include "Misc/AssertionMacros.h"
 #include "ScopedTransaction.h"
@@ -10,6 +11,16 @@
 /////////////////////////////////////////////////////
 #define LOCTEXT_NAMESPACE "EditorUtility"
 
+UWidget* UEditorUtilityWidget::FindChildWidgetByName(FName WidgetName) const
+{
+	if (WidgetTree)
+	{
+		return WidgetTree->FindWidget(WidgetName);
+	}
+
+	return nullptr;
+}
+
 void UEditorUtilityWidget::ExecuteDefaultAction()
 {
 	check(bAutoRunDefaultAction);
