Blutility: Add editor scripting functions for modifying Widgets in a Widget Blueprint's Widget Tree

#rnx

diff --git a/Engine/Source/Editor/Blutility/Private/EditorUtilityLibrary.cpp b/Engine/Source/Editor/Blutility/Private/EditorUtilityLibrary.cpp
--- a/Engine/Source/Editor/Blutility/Private/EditorUtilityLibrary.cpp
+++ b/Engine/Source/Editor/Blutility/Private/EditorUtilityLibrary.cpp
@@ -419,6 +419,29 @@ void UEditorUtilityLibrary::ConvertToEditorUtilityWidget(UWidgetBlueprint* Widge
 	}
 }
 
+void UEditorUtilityLibrary::CastToWidgetBlueprint(UObject* Object, ECastToWidgetBlueprintCases& Branches, UWidgetBlueprint*& AsWidgetBlueprint)
+{
+	AsWidgetBlueprint = Cast<UWidgetBlueprint>(Object);
+	Branches = (AsWidgetBlueprint == nullptr) ? ECastToWidgetBlueprintCases::CastFailed : ECastToWidgetBlueprintCases::CastSucceeded;
+}
+
+UWidget* UEditorUtilityLibrary::FindSourceWidgetByName(UWidgetBlueprint* WidgetBlueprint, FName WidgetName)
+{
+	UWidget* FoundWidget = nullptr;
+
+	if (WidgetBlueprint)
+	{
+		WidgetBlueprint->ForEachSourceWidget([&WidgetName, &FoundWidget](UWidget* Widget)
+		{
+			if (Widget->GetName() == WidgetName)
+			{
+				FoundWidget = Widget;
+			}
+		});
+	}
+
+	return FoundWidget;
+}
 #endif
 
 #undef LOCTEXT_NAMESPACE
diff --git a/Engine/Source/Editor/Blutility/Public/EditorUtilityLibrary.h b/Engine/Source/Editor/Blutility/Public/EditorUtilityLibrary.h
--- a/Engine/Source/Editor/Blutility/Public/EditorUtilityLibrary.h
+++ b/Engine/Source/Editor/Blutility/Public/EditorUtilityLibrary.h
@@ -18,6 +18,8 @@
 
 #include "EditorUtilityLibrary.generated.h"
 
+class UWidget;
+class UEditorUtilityWidgetBlueprint;
 class AActor;
 class UClass;
 class UEditorPerProjectUserSettings;
@@ -26,6 +28,13 @@ struct FAssetData;
 struct FContentBrowserItemPath;
 struct FFrame;
 
+UENUM(BlueprintType)
+enum class ECastToWidgetBlueprintCases : uint8
+{
+	CastSucceeded,
+	CastFailed
+};
+
 UCLASS()
 class BLUTILITY_API UEditorUtilityBlueprintAsyncActionBase : public UBlueprintAsyncActionBase
 {
@@ -203,5 +212,12 @@ public:
 	
 	UFUNCTION(BlueprintCallable, Category = "Editor Scripting | Content Browser")
 	static void ConvertToEditorUtilityWidget(class UWidgetBlueprint* WidgetBP);
+	
+	UFUNCTION(BlueprintCallable, Category = "Editor Scripting | Widget Blueprint", Meta = (ExpandEnumAsExecs = "Branches"))
+	static void CastToWidgetBlueprint(UObject* Object, ECastToWidgetBlueprintCases& Branches, UWidgetBlueprint*& AsWidgetBlueprint);
+
+	/** Searches the blueprint's widget hierarchy for a widget with the specified name */
+	UFUNCTION(BlueprintCallable, BlueprintPure = false, Category = "Editor Scripting | Widget Blueprint")
+	static UWidget* FindSourceWidgetByName(UWidgetBlueprint* WidgetBlueprint, FName WidgetName);
 #endif
 };
