Unregister handlers when the tab is removed.
#13007

#tests Verified construction/destruction, tracking instances on the LevelEditor object
#rnx

diff --git a/Engine/Source/Editor/Blutility/Private/EditorUtilityWidgetBlueprint.cpp b/Engine/Source/Editor/Blutility/Private/EditorUtilityWidgetBlueprint.cpp
--- a/Engine/Source/Editor/Blutility/Private/EditorUtilityWidgetBlueprint.cpp
+++ b/Engine/Source/Editor/Blutility/Private/EditorUtilityWidgetBlueprint.cpp
@@ -235,6 +235,13 @@ void UEditorUtilityWidgetBlueprint::UpdateRespawnListIfNeeded(TSharedRef<SDockTa
 		CreatedUMGWidget->Rename(nullptr, GetTransientPackage());
 		CreatedUMGWidget = nullptr;
 	}
+
+	if (CreatedTab.IsValid())
+	{
+		FLevelEditorModule& LevelEditor = FModuleManager::LoadModuleChecked<FLevelEditorModule>("LevelEditor");
+		LevelEditor.OnMapChanged().RemoveAll(this);
+	}
+
 }
 
 void UEditorUtilityWidgetBlueprint::GetReparentingRules(TSet< const UClass* >& AllowedChildrenOfClasses, TSet< const UClass* >& DisallowedChildrenOfClasses) const
