EditorUtilityTask - Add bool flag to FinishExecutingTask to allow reporting a failed task

diff --git a/Engine/Source/Editor/Blutility/Classes/EditorUtilityTask.h b/Engine/Source/Editor/Blutility/Classes/EditorUtilityTask.h
--- a/Engine/Source/Editor/Blutility/Classes/EditorUtilityTask.h
+++ b/Engine/Source/Editor/Blutility/Classes/EditorUtilityTask.h
@@ -43,7 +43,7 @@ public:
 
 public:
 	UFUNCTION(BlueprintCallable, Category=Task)
-	void FinishExecutingTask();
+	void FinishExecutingTask(const bool bSuccess = true);
 
 	UFUNCTION(BlueprintCallable, Category = Task)
 	void SetTaskNotificationText(const FText& Text);
diff --git a/Engine/Source/Editor/Blutility/Private/EditorUtilityTask.cpp b/Engine/Source/Editor/Blutility/Private/EditorUtilityTask.cpp
--- a/Engine/Source/Editor/Blutility/Private/EditorUtilityTask.cpp
+++ b/Engine/Source/Editor/Blutility/Private/EditorUtilityTask.cpp
@@ -57,7 +57,7 @@ void UEditorUtilityTask::StartExecutingTask()
 	ReceiveBeginExecution();
 }
 
-void UEditorUtilityTask::FinishExecutingTask()
+void UEditorUtilityTask::FinishExecutingTask(const bool bSuccess)
 {
 	SetTaskNotificationText(LOCTEXT("TaskComplete", "Complete"));
 
@@ -68,7 +68,7 @@ void UEditorUtilityTask::FinishExecutingTask()
 
 	if (TaskNotification.IsValid())
 	{
-		TaskNotification->SetComplete(true);
+		TaskNotification->SetComplete(bSuccess);
 		TaskNotification.Reset();
 	}
 
