Blutility: Add UEditorUtilityLibrary::AddSourceWidget to allow adding widgets to a widget blueprint from editor scripts

#rnx

diff --git a/Engine/Source/Editor/Blutility/Private/EditorUtilityLibrary.cpp b/Engine/Source/Editor/Blutility/Private/EditorUtilityLibrary.cpp
--- a/Engine/Source/Editor/Blutility/Private/EditorUtilityLibrary.cpp
+++ b/Engine/Source/Editor/Blutility/Private/EditorUtilityLibrary.cpp
@@ -13,6 +13,7 @@
 #include "IAssetTools.h"
 #include "EditorUtilitySubsystem.h"
 #include "Kismet/KismetSystemLibrary.h"
+#include "Kismet2/BlueprintEditorUtils.h"
 #include "Serialization/ArchiveReplaceObjectRef.h"
 #include "Serialization/FindReferencersArchive.h"
 #include "Templates/SubclassOf.h"
@@ -442,6 +443,58 @@ UWidget* UEditorUtilityLibrary::FindSourceWidgetByName(UWidgetBlueprint* WidgetB
 
 	return FoundWidget;
 }
+
+UWidget* UEditorUtilityLibrary::AddSourceWidget(UWidgetBlueprint* WidgetBlueprint, TSubclassOf<UWidget> WidgetClass, FName WidgetName, FName WidgetParentName)
+{
+	if (!WidgetBlueprint || !WidgetBlueprint->WidgetTree)
+	{
+		FFrame::KismetExecutionMessage(TEXT("Attempting to add a widget to an invalid Widget Blueprint"), ELogVerbosity::Error);
+		return nullptr;
+	}
+
+	if (!WidgetClass)
+	{
+		FFrame::KismetExecutionMessage(TEXT("Attempting to add an invalid widget class to a Widget Blueprint"), ELogVerbosity::Error);
+		return nullptr;
+	}
+
+	UPanelWidget* ParentWidget = WidgetParentName.IsNone() ? nullptr : Cast<UPanelWidget>(FindSourceWidgetByName(WidgetBlueprint, WidgetParentName));
+	if (!ParentWidget && !WidgetParentName.IsNone())
+	{
+		FFrame::KismetExecutionMessage(TEXT("Attempting to add a widget to an invalid parent"), ELogVerbosity::Error);
+		return nullptr;
+	}
+
+	if (WidgetParentName.IsNone() && WidgetBlueprint->WidgetTree->RootWidget != nullptr)
+	{
+		FFrame::KismetExecutionMessage(TEXT("Attempting to set the root widget when one already exists"), ELogVerbosity::Error);
+		return nullptr;
+	}
+
+	const bool bDoesWidgetNameExist = FindSourceWidgetByName(WidgetBlueprint, WidgetName) != nullptr;
+	const FName UniqueWidgetName = bDoesWidgetNameExist ? MakeUniqueObjectName(WidgetBlueprint->WidgetTree, WidgetClass, WidgetName) : WidgetName;
+	UWidget* NewWidget = WidgetBlueprint->WidgetTree->ConstructWidget(WidgetClass, UniqueWidgetName);
+	if (!NewWidget)
+	{
+		FFrame::KismetExecutionMessage(TEXT("Failed to create a widget to add to a Widget Blueprint"), ELogVerbosity::Error);
+		return nullptr;
+	}
+
+	WidgetBlueprint->Modify();
+
+	if (ParentWidget)
+	{
+		ParentWidget->AddChild(NewWidget);
+	}
+	else
+	{
+		WidgetBlueprint->WidgetTree->RootWidget = NewWidget;
+	}
+
+	WidgetBlueprint->OnVariableAdded(NewWidget->GetFName());
+	FBlueprintEditorUtils::MarkBlueprintAsStructurallyModified(WidgetBlueprint);
+	return NewWidget;
+}
 #endif
 
 #undef LOCTEXT_NAMESPACE
diff --git a/Engine/Source/Editor/Blutility/Public/EditorUtilityLibrary.h b/Engine/Source/Editor/Blutility/Public/EditorUtilityLibrary.h
--- a/Engine/Source/Editor/Blutility/Public/EditorUtilityLibrary.h
+++ b/Engine/Source/Editor/Blutility/Public/EditorUtilityLibrary.h
@@ -219,5 +219,18 @@ public:
 	/** Searches the blueprint's widget hierarchy for a widget with the specified name */
 	UFUNCTION(BlueprintCallable, BlueprintPure = false, Category = "Editor Scripting | Widget Blueprint")
 	static UWidget* FindSourceWidgetByName(UWidgetBlueprint* WidgetBlueprint, FName WidgetName);
+
+	/**
+	 * Create a new widget and add to the specific widget blueprint's widget tree
+	 *
+	 * @param WidgetBlueprint The widget blueprint to add a widget to
+	 * @param WidgetClass The widget class to add to the widget blueprint
+	 * @param WidgetName The name to give the new widget
+	 * @param WidgetParentName The name of the existing widget that will hold the new widget. Must be an existing Panel Widget or none of if the widget tree is empty and the new widget will become the RootWidget.
+	 *
+	 * @return The widget that was created and added to the widget blueprint
+	 */
+	UFUNCTION(BlueprintCallable, BlueprintPure = false, Category = "Editor Scripting | Widget Blueprint")
+	static UWidget* AddSourceWidget(UWidgetBlueprint* WidgetBlueprint, TSubclassOf<UWidget> WidgetClass, FName WidgetName, FName WidgetParentName);
 #endif
 };
