Improvements to editor utility start-up and unload management

Editor utilities may now opt-in to running on start-up via a new 'Run on Start-up' option on their Class Defaults. Like the existing StartupObjects config array method, this requires that the Blueprint implement a Run function/event, and the editor will wait until the main editor frame is available before running this function/event. The new  'Run on Start-up' option also requires that the initial asset registry scan has completed, so may run later than the StartupObjects config array method.

The subsystem that manages start-up editor utilities will also now clean itself up correctly when an editor utility blueprint asset is unloaded. This includes trying to remove any editor utility generated menus, however this support is currently limited to the data driven method of defining tools menus (via the Data property of EditorUtilityToolMenuEntry).

diff --git a/Engine/Source/Developer/ToolMenus/Private/ToolMenu.cpp b/Engine/Source/Developer/ToolMenus/Private/ToolMenu.cpp
--- a/Engine/Source/Developer/ToolMenus/Private/ToolMenu.cpp
+++ b/Engine/Source/Developer/ToolMenus/Private/ToolMenu.cpp
@@ -314,6 +314,19 @@ void UToolMenu::AddMenuEntryObject(UToolMenuEntryScript* InObject)
 	}
 }
 
+void UToolMenu::RemoveMenuEntryObject(UToolMenuEntryScript* InObject)
+{
+	if (FToolMenuSection* Section = FindSection(InObject->Data.Section))
+	{
+		Section->RemoveEntryObject(InObject);
+	}
+
+	if (MenuType == EMultiBoxType::MenuBar || MenuType == EMultiBoxType::ToolBar)
+	{
+		UToolMenus::Get()->RefreshAllWidgets();
+	}
+}
+
 UToolMenu* UToolMenu::AddSubMenuScript(const FName InOwner, const FName SectionName, const FName InName, const FText& InLabel, const FText& InToolTip)
 {
 	return AddSubMenu(InOwner, SectionName, InName, InLabel, InToolTip);
diff --git a/Engine/Source/Developer/ToolMenus/Private/ToolMenuEntryScript.cpp b/Engine/Source/Developer/ToolMenus/Private/ToolMenuEntryScript.cpp
--- a/Engine/Source/Developer/ToolMenus/Private/ToolMenuEntryScript.cpp
+++ b/Engine/Source/Developer/ToolMenus/Private/ToolMenuEntryScript.cpp
@@ -138,7 +138,16 @@ FSlateIcon UToolMenuEntryScript::GetSlateIcon(const FToolMenuContext& Context) c
 
 void UToolMenuEntryScript::RegisterMenuEntry()
 {
-	UToolMenus::AddMenuEntryObject(this);
+	bHasRegisteredMenuEntry = UToolMenus::AddMenuEntryObject(this);
+}
+
+void UToolMenuEntryScript::UnregisterMenuEntry()
+{
+	if (bHasRegisteredMenuEntry)
+	{
+		UToolMenus::RemoveMenuEntryObject(this);
+		bHasRegisteredMenuEntry = false;
+	}
 }
 
 void UToolMenuEntryScript::InitEntry(const FName OwnerName, const FName Menu, const FName Section, const FName Name, const FText& Label, const FText& ToolTip)
diff --git a/Engine/Source/Developer/ToolMenus/Private/ToolMenuSection.cpp b/Engine/Source/Developer/ToolMenus/Private/ToolMenuSection.cpp
--- a/Engine/Source/Developer/ToolMenus/Private/ToolMenuSection.cpp
+++ b/Engine/Source/Developer/ToolMenus/Private/ToolMenuSection.cpp
@@ -202,6 +202,11 @@ int32 FToolMenuSection::RemoveEntry(const FName InName)
 	return Blocks.RemoveAll([InName](const FToolMenuEntry& Block) { return Block.Name == InName; });
 }
 
+int32 FToolMenuSection::RemoveEntryObject(const UToolMenuEntryScript* InObject)
+{
+	return Blocks.RemoveAll([InObject](const FToolMenuEntry& Block) { return Block.ScriptObject == InObject && Block.Name == InObject->Data.Name; });
+}
+
 int32 FToolMenuSection::RemoveEntriesByOwner(const FToolMenuOwner InOwner)
 {
 	if (InOwner != FToolMenuOwner())
diff --git a/Engine/Source/Developer/ToolMenus/Private/ToolMenus.cpp b/Engine/Source/Developer/ToolMenus/Private/ToolMenus.cpp
--- a/Engine/Source/Developer/ToolMenus/Private/ToolMenus.cpp
+++ b/Engine/Source/Developer/ToolMenus/Private/ToolMenus.cpp
@@ -3265,6 +3265,16 @@ bool UToolMenus::AddMenuEntryObject(UToolMenuEntryScript* MenuEntryObject)
 	return true;
 }
 
+bool UToolMenus::RemoveMenuEntryObject(UToolMenuEntryScript* MenuEntryObject)
+{
+	if (UToolMenu* Menu = UToolMenus::Get()->FindMenu(MenuEntryObject->Data.Menu))
+	{
+		Menu->RemoveMenuEntryObject(MenuEntryObject);
+		return true;
+	}
+	return false;
+}
+
 void UToolMenus::SetSectionLabel(const FName MenuName, const FName SectionName, const FText Label)
 {
 	ExtendMenu(MenuName)->FindOrAddSection(SectionName).Label = TAttribute<FText>(Label);
diff --git a/Engine/Source/Developer/ToolMenus/Public/ToolMenu.h b/Engine/Source/Developer/ToolMenus/Public/ToolMenu.h
--- a/Engine/Source/Developer/ToolMenus/Public/ToolMenu.h
+++ b/Engine/Source/Developer/ToolMenus/Public/ToolMenu.h
@@ -48,6 +48,9 @@ public:
 	UFUNCTION(BlueprintCallable, Category = "Tool Menus")
 	void AddMenuEntryObject(UToolMenuEntryScript* InObject);
 
+	UFUNCTION(BlueprintCallable, Category = "Tool Menus")
+	void RemoveMenuEntryObject(UToolMenuEntryScript* InObject);
+
 	UFUNCTION(BlueprintCallable, Category = "Tool Menus", meta = ( ScriptName = "AddSubMenu", AutoCreateRefTerm = "Label,ToolTip" ))
 	UToolMenu* AddSubMenuScript(const FName Owner, const FName SectionName, const FName Name, const FText& Label, const FText& ToolTip = FText());
 
diff --git a/Engine/Source/Developer/ToolMenus/Public/ToolMenuEntryScript.h b/Engine/Source/Developer/ToolMenus/Public/ToolMenuEntryScript.h
--- a/Engine/Source/Developer/ToolMenus/Public/ToolMenuEntryScript.h
+++ b/Engine/Source/Developer/ToolMenus/Public/ToolMenuEntryScript.h
@@ -155,6 +155,9 @@ public:
 	UFUNCTION(BlueprintCallable, Category = "Advanced")
 	void RegisterMenuEntry();
 
+	UFUNCTION(BlueprintCallable, Category = "Advanced")
+	void UnregisterMenuEntry();
+
 	UFUNCTION(BlueprintCallable, Category = "Advanced")
 	void InitEntry(const FName OwnerName, const FName Menu, const FName Section, const FName Name, const FText& Label = FText(), const FText& ToolTip = FText());
 
@@ -193,4 +196,7 @@ public:
 
 	UPROPERTY(EditAnywhere, BlueprintReadWrite, Category = "Data")
 	FToolMenuEntryScriptData Data;
+
+	UPROPERTY(Transient)
+	bool bHasRegisteredMenuEntry = false;
 };
diff --git a/Engine/Source/Developer/ToolMenus/Public/ToolMenuSection.h b/Engine/Source/Developer/ToolMenus/Public/ToolMenuSection.h
--- a/Engine/Source/Developer/ToolMenus/Public/ToolMenuSection.h
+++ b/Engine/Source/Developer/ToolMenus/Public/ToolMenuSection.h
@@ -99,6 +99,7 @@ private:
 	void InitGeneratedSectionCopy(const FToolMenuSection& Source, FToolMenuContext& InContext);
 
 	int32 RemoveEntry(const FName InName);
+	int32 RemoveEntryObject(const UToolMenuEntryScript* InObject);
 	int32 RemoveEntriesByOwner(const FToolMenuOwner InOwner);
 
 	int32 IndexOfBlock(const FName InName) const;
diff --git a/Engine/Source/Developer/ToolMenus/Public/ToolMenus.h b/Engine/Source/Developer/ToolMenus/Public/ToolMenus.h
--- a/Engine/Source/Developer/ToolMenus/Public/ToolMenus.h
+++ b/Engine/Source/Developer/ToolMenus/Public/ToolMenus.h
@@ -202,6 +202,10 @@ public:
 	UFUNCTION(BlueprintCallable, Category = "Tool Menus")
 	static bool AddMenuEntryObject(UToolMenuEntryScript* MenuEntryObject);
 
+	/** Unregisters menu entry object from blueprint/script */
+	UFUNCTION(BlueprintCallable, Category = "Tool Menus")
+	static bool RemoveMenuEntryObject(UToolMenuEntryScript* MenuEntryObject);
+
 	/** Removes all entries that were registered under a specific owner name */
 	UFUNCTION(BlueprintCallable, Category = "Tool Menus")
 	void UnregisterOwnerByName(FName InOwnerName);
diff --git a/Engine/Source/Editor/Blutility/Classes/EditorUtilityObject.h b/Engine/Source/Editor/Blutility/Classes/EditorUtilityObject.h
--- a/Engine/Source/Editor/Blutility/Classes/EditorUtilityObject.h
+++ b/Engine/Source/Editor/Blutility/Classes/EditorUtilityObject.h
@@ -29,4 +29,9 @@ class BLUTILITY_API UEditorUtilityObject : public UObject, public IAssetRegistry
 		return true; 
 	}
 	//~ End IAssetRegistryTagProviderInterface interface
+
+protected:
+	/** Run this editor utility on start-up (after asset discovery)? */
+	UPROPERTY(Category=Settings, EditDefaultsOnly, AssetRegistrySearchable, DisplayName="Run on Start-up")
+	bool bRunEditorUtilityOnStartup = false;
 };
diff --git a/Engine/Source/Editor/Blutility/Classes/EditorUtilityTask.h b/Engine/Source/Editor/Blutility/Classes/EditorUtilityTask.h
--- a/Engine/Source/Editor/Blutility/Classes/EditorUtilityTask.h
+++ b/Engine/Source/Editor/Blutility/Classes/EditorUtilityTask.h
@@ -5,6 +5,7 @@
 #include "CoreMinimal.h"
 #include "Delegates/Delegate.h"
 #include "Engine/World.h"
+#include "IAssetRegistryTagProviderInterface.h"
 #include "Misc/AsyncTaskNotification.h"
 #include "Templates/UniquePtr.h"
 #include "UObject/Object.h"
@@ -26,7 +27,7 @@ DECLARE_MULTICAST_DELEGATE_OneParam(FOnEditorUtilityTaskDelegate, UEditorUtility
  * 
  */
 UCLASS(Abstract, Blueprintable, meta = (ShowWorldContextPin))
-class BLUTILITY_API UEditorUtilityTask : public UObject
+class BLUTILITY_API UEditorUtilityTask : public UObject, public IAssetRegistryTagProviderInterface
 {
 	GENERATED_BODY()
 
@@ -39,6 +40,13 @@ public:
 	UFUNCTION()
 	void Run();
 
+	//~ Begin IAssetRegistryTagProviderInterface interface
+	virtual bool ShouldAddCDOTagsToBlueprintClass() const override
+	{
+		return true;
+	}
+	//~ End IAssetRegistryTagProviderInterface interface
+
 	virtual UWorld* GetWorld() const override;
 
 public:
@@ -75,6 +83,11 @@ private:
 	// Calls BeginExecution() and ReceiveBeginExecution()
 	void StartExecutingTask();
 
+protected:
+	/** Run this editor utility on start-up (after asset discovery)? */
+	UPROPERTY(Category=Settings, EditDefaultsOnly, AssetRegistrySearchable, DisplayName="Run on Start-up")
+	bool bRunEditorUtilityOnStartup = false;
+
 private:
 	UPROPERTY(Transient)
 	TObjectPtr<UEditorUtilitySubsystem> MyTaskManager;
diff --git a/Engine/Source/Editor/Blutility/Classes/EditorUtilityToolMenu.h b/Engine/Source/Editor/Blutility/Classes/EditorUtilityToolMenu.h
--- a/Engine/Source/Editor/Blutility/Classes/EditorUtilityToolMenu.h
+++ b/Engine/Source/Editor/Blutility/Classes/EditorUtilityToolMenu.h
@@ -3,6 +3,7 @@
 #pragma once
 
 #include "CoreMinimal.h"
+#include "IAssetRegistryTagProviderInterface.h"
 #include "UObject/ScriptMacros.h"
 #include "ToolMenuEntryScript.h"
 #include "ToolMenuSection.h"
@@ -10,14 +11,38 @@
 
 
 UCLASS(Blueprintable, abstract)
-class BLUTILITY_API UEditorUtilityToolMenuEntry : public UToolMenuEntryScript
+class BLUTILITY_API UEditorUtilityToolMenuEntry : public UToolMenuEntryScript, public IAssetRegistryTagProviderInterface
 {
 	GENERATED_BODY()
+
+	//~ Begin IAssetRegistryTagProviderInterface interface
+	virtual bool ShouldAddCDOTagsToBlueprintClass() const override
+	{
+		return true;
+	}
+	//~ End IAssetRegistryTagProviderInterface interface
+
+protected:
+	/** Run this editor utility on start-up (after asset discovery)? */
+	UPROPERTY(Category=Settings, EditDefaultsOnly, AssetRegistrySearchable, DisplayName="Run on Start-up")
+	bool bRunEditorUtilityOnStartup = false;
 };
 
 UCLASS(Blueprintable, abstract)
-class BLUTILITY_API UEditorUtilityToolMenuSection : public UToolMenuSectionDynamic
+class BLUTILITY_API UEditorUtilityToolMenuSection : public UToolMenuSectionDynamic, public IAssetRegistryTagProviderInterface
 {
 	GENERATED_BODY()
+
+	//~ Begin IAssetRegistryTagProviderInterface interface
+	virtual bool ShouldAddCDOTagsToBlueprintClass() const override
+	{
+		return true;
+	}
+	//~ End IAssetRegistryTagProviderInterface interface
+
+protected:
+	/** Run this editor utility on start-up (after asset discovery)? */
+	UPROPERTY(Category=Settings, EditDefaultsOnly, AssetRegistrySearchable, DisplayName="Run on Start-up")
+	bool bRunEditorUtilityOnStartup = false;
 };
 
diff --git a/Engine/Source/Editor/Blutility/Classes/EditorUtilityWidget.h b/Engine/Source/Editor/Blutility/Classes/EditorUtilityWidget.h
--- a/Engine/Source/Editor/Blutility/Classes/EditorUtilityWidget.h
+++ b/Engine/Source/Editor/Blutility/Classes/EditorUtilityWidget.h
@@ -9,6 +9,7 @@
 #include "Blueprint/UserWidget.h"
 #include "Containers/UnrealString.h"
 #include "CoreMinimal.h"
+#include "IAssetRegistryTagProviderInterface.h"
 #include "UObject/Object.h"
 #include "UObject/ObjectMacros.h"
 #include "UObject/ScriptMacros.h"
@@ -21,7 +22,7 @@ class UEditorPerProjectUserSettings;
 class UObject;
 
 UCLASS(Abstract, meta = (ShowWorldContextPin), config = Editor)
-class BLUTILITY_API UEditorUtilityWidget : public UUserWidget
+class BLUTILITY_API UEditorUtilityWidget : public UUserWidget, public IAssetRegistryTagProviderInterface
 {
 	GENERATED_BODY()
 
@@ -30,6 +31,13 @@ public:
 	UFUNCTION(BlueprintCallable, BlueprintImplementableEvent, Category = "Editor")
 	void Run();
 
+	//~ Begin IAssetRegistryTagProviderInterface interface
+	virtual bool ShouldAddCDOTagsToBlueprintClass() const override
+	{
+		return true;
+	}
+	//~ End IAssetRegistryTagProviderInterface interface
+
 	UFUNCTION(BlueprintCallable, BlueprintPure = false, Category = "Editor Utility Widget")
 	UWidget* FindChildWidgetByName(FName WidgetName) const;
 
@@ -69,4 +77,8 @@ protected:
 	// Should this blueprint automatically run OnDefaultActionClicked, or should it open up a details panel to edit properties and/or offer multiple buttons
 	UPROPERTY(Category = Settings, EditDefaultsOnly, BlueprintReadOnly)
 	bool bAutoRunDefaultAction;
+
+	/** Run this editor utility on start-up (after asset discovery)? */
+	UPROPERTY(Category=Settings, EditDefaultsOnly, AssetRegistrySearchable, DisplayName="Run on Start-up")
+	bool bRunEditorUtilityOnStartup = false;
 };
diff --git a/Engine/Source/Editor/Blutility/Private/EditorUtilitySubsystem.cpp b/Engine/Source/Editor/Blutility/Private/EditorUtilitySubsystem.cpp
--- a/Engine/Source/Editor/Blutility/Private/EditorUtilitySubsystem.cpp
+++ b/Engine/Source/Editor/Blutility/Private/EditorUtilitySubsystem.cpp
@@ -4,12 +4,12 @@
 
 #include "AssetRegistry/ARFilter.h"
 #include "AssetRegistry/AssetData.h"
-#include "AssetRegistry/AssetRegistryModule.h"
 #include "AssetRegistry/IAssetRegistry.h"
 #include "CoreGlobals.h"
 #include "Editor.h"
 #include "EditorUtilityCommon.h"
 #include "EditorUtilityTask.h"
+#include "EditorUtilityToolMenu.h"
 #include "EditorUtilityWidgetBlueprint.h"
 #include "Engine/Blueprint.h"
 #include "Engine/BlueprintCore.h"
@@ -23,6 +23,7 @@
 #include "Interfaces/IMainFrameModule.h"
 #include "Internationalization/Internationalization.h"
 #include "Internationalization/Text.h"
+#include "Kismet2/KismetEditorUtilities.h"
 #include "LevelEditor.h"
 #include "Logging/LogCategory.h"
 #include "Logging/LogMacros.h"
@@ -47,6 +48,28 @@ class UWorld;
 
 #define LOCTEXT_NAMESPACE "EditorUtilitySubsystem"
 
+namespace UE::EditorUtility::Private
+{
+
+static const FName RunFunctionName = "Run";
+static const FName RunOnStartupTagName = "bRunEditorUtilityOnStartup";
+
+UClass* GetRunnableClassForAsset(UObject* Asset)
+{
+	if (UClass* Class = Cast<UClass>(Asset))
+	{
+		return Class;
+	}
+
+	if (UBlueprint* Blueprint = Cast<UBlueprint>(Asset))
+	{
+		return Blueprint->GeneratedClass;
+	}
+
+	return Asset->GetClass();
+}
+
+}
 
 UEditorUtilitySubsystem::UEditorUtilitySubsystem()
 	: UEditorSubsystem()
@@ -94,6 +117,8 @@ void UEditorUtilitySubsystem::Initialize(FSubsystemCollectionBase& Collection)
 
 	FLevelEditorModule& LevelEditor = FModuleManager::LoadModuleChecked<FLevelEditorModule>("LevelEditor");
 	LevelEditor.OnMapChanged().AddUObject(this, &UEditorUtilitySubsystem::OnMapChanged);
+
+	FKismetEditorUtilities::OnBlueprintGeneratedClassUnloaded.AddUObject(this, &UEditorUtilitySubsystem::OnBlueprintGeneratedClassUnloaded);
 }
 
 void UEditorUtilitySubsystem::Deinitialize()
@@ -122,6 +147,8 @@ void UEditorUtilitySubsystem::Deinitialize()
 	{
 		LevelEditor->OnMapChanged().RemoveAll(this);
 	}
+
+	FKismetEditorUtilities::OnBlueprintGeneratedClassUnloaded.RemoveAll(this);
 }
 
 void UEditorUtilitySubsystem::AddReferencedObjects(UObject* InThis, FReferenceCollector& Collector)
@@ -142,6 +169,7 @@ void UEditorUtilitySubsystem::MainFrameCreationFinished(TSharedPtr<SWindow> InRo
 
 void UEditorUtilitySubsystem::HandleStartup()
 {
+	// Handle Blueprints in the StartupObjects config list
 	for (const FSoftObjectPath& ObjectPath : StartupObjects)
 	{
 		UObject* Object = GetValid(ObjectPath.TryLoad());
@@ -151,7 +179,96 @@ void UEditorUtilitySubsystem::HandleStartup()
 			continue;
 		}
 		
-		TryRun(ObjectPath.TryLoad());
+		TryRun(Object);
+	}
+
+	// Handle Blueprints with bRunEditorUtilityOnStartup set to true
+	if (IAssetRegistry& AssetRegistry = IAssetRegistry::GetChecked();
+		AssetRegistry.IsLoadingAssets())
+	{
+		AssetRegistry.OnFilesLoaded().AddUObject(this, &UEditorUtilitySubsystem::HandleStartupAssets);
+	}
+	else
+	{
+		HandleStartupAssets();
+	}
+}
+
+void UEditorUtilitySubsystem::HandleStartupAssets()
+{
+	IAssetRegistry& AssetRegistry = IAssetRegistry::GetChecked();
+	check(!AssetRegistry.IsLoadingAssets());
+
+	// Handle the current set of assets
+	{
+		FARFilter Filter;
+		Filter.bRecursiveClasses = true;
+		Filter.ClassPaths.Add(UBlueprintCore::StaticClass()->GetClassPathName());
+		Filter.ClassPaths.Add(UBlueprintGeneratedClass::StaticClass()->GetClassPathName());
+		Filter.TagsAndValues.Emplace(UE::EditorUtility::Private::RunOnStartupTagName, TEXT("True"));
+
+		TArray<FAssetData> StartupAssets;
+		AssetRegistry.GetAssets(Filter, StartupAssets);
+
+		for (const FAssetData& StartupAsset : StartupAssets)
+		{
+			UObject* Object = StartupAsset.GetAsset();
+			if (!Object || Object->IsUnreachable())
+			{
+				UE_LOG(LogEditorUtilityBlueprint, Warning, TEXT("Could not load: %s"), *StartupAsset.GetSoftObjectPath().ToString());
+				continue;
+			}
+
+			TryRun(Object);
+		}
+	}
+
+	// Watch for new assets being added
+	AssetRegistry.OnAssetsAdded().AddUObject(this, &UEditorUtilitySubsystem::HandleDynamicStartupAssets);
+}
+
+void UEditorUtilitySubsystem::HandleDynamicStartupAssets(TConstArrayView<FAssetData> InAssets)
+{
+	for (const FAssetData& Asset : InAssets)
+	{
+		if (const UClass* AssetClass = Asset.GetClass();
+			AssetClass && (AssetClass->IsChildOf<UBlueprintCore>() || AssetClass->IsChildOf<UBlueprintGeneratedClass>()))
+		{
+			if (bool bRunOnStartup = false;
+				Asset.GetTagValue<bool>(UE::EditorUtility::Private::RunOnStartupTagName, bRunOnStartup) && bRunOnStartup)
+			{
+				UObject* Object = Asset.GetAsset();
+				if (!Object || Object->IsUnreachable())
+				{
+					UE_LOG(LogEditorUtilityBlueprint, Warning, TEXT("Could not load: %s"), *Asset.GetSoftObjectPath().ToString());
+					continue;
+				}
+
+				TryRun(Object);
+			}
+		}
+	}
+}
+
+void UEditorUtilitySubsystem::OnBlueprintGeneratedClassUnloaded(UBlueprintGeneratedClass* BPGC)
+{
+	if (BPGC->IsChildOf<UEditorUtilityToolMenuEntry>())
+	{
+		// Unregister any menus related to this BPGC
+		TArray<UObject*> MenuInstanceObjects;
+		GetObjectsOfClass(BPGC, MenuInstanceObjects, /*bIncludeDerivedClasses*/false, RF_NewerVersionExists, EInternalObjectFlags::Garbage);
+
+		for (UObject* MenuInstanceObject : MenuInstanceObjects)
+		{
+			UEditorUtilityToolMenuEntry* MenuInstance = CastChecked<UEditorUtilityToolMenuEntry>(MenuInstanceObject);
+			MenuInstance->UnregisterMenuEntry();
+		}
+	}
+
+	ReleaseInstanceOfAsset(BPGC);
+	if (BPGC->ClassGeneratedBy)
+	{
+		ReleaseInstanceOfAsset(BPGC->ClassGeneratedBy);
 	}
 }
 
@@ -163,12 +280,7 @@ bool UEditorUtilitySubsystem::TryRun(UObject* Asset)
 		return false;
 	}
 
-	UClass* ObjectClass = Asset->GetClass();
-	if (UBlueprint* Blueprint = Cast<UBlueprint>(Asset))
-	{
-		ObjectClass = Blueprint->GeneratedClass;
-	}
-
+	UClass* ObjectClass = UE::EditorUtility::Private::GetRunnableClassForAsset(Asset);
 	if (!ObjectClass)
 	{
 		UE_LOG(LogEditorUtilityBlueprint, Warning, TEXT("Missing class: %s"), *Asset->GetPathName());
@@ -181,8 +293,7 @@ bool UEditorUtilitySubsystem::TryRun(UObject* Asset)
 		return false;
 	}
 
-	static const FName RunFunctionName("Run");
-	UFunction* RunFunction = ObjectClass->FindFunctionByName(RunFunctionName);
+	UFunction* RunFunction = ObjectClass->FindFunctionByName(UE::EditorUtility::Private::RunFunctionName);
 	if (RunFunction && RunFunction->ParmsSize == 0)
 	{
 		UObject* Instance = NewObject<UObject>(this, ObjectClass);
@@ -202,8 +313,7 @@ bool UEditorUtilitySubsystem::TryRun(UObject* Asset)
 
 bool UEditorUtilitySubsystem::TryRunClass(UClass* ObjectClass)
 {
-	static const FName RunFunctionName("Run");
-	UFunction* RunFunction = ObjectClass->FindFunctionByName(RunFunctionName);
+	UFunction* RunFunction = ObjectClass->FindFunctionByName(UE::EditorUtility::Private::RunFunctionName);
 	if (RunFunction && RunFunction->ParmsSize == 0)
 	{
 		UObject* Instance = NewObject<UObject>(this, ObjectClass);
@@ -223,12 +333,7 @@ bool UEditorUtilitySubsystem::TryRunClass(UClass* ObjectClass)
 
 bool UEditorUtilitySubsystem::CanRun(UObject* Asset) const
 {
-	UClass* ObjectClass = Asset->GetClass();
-	if (UBlueprint* Blueprint = Cast<UBlueprint>(Asset))
-	{
-		ObjectClass = Blueprint->GeneratedClass;
-	}
-
+	UClass* ObjectClass = UE::EditorUtility::Private::GetRunnableClassForAsset(Asset);
 	if (ObjectClass)
 	{
 		if (ObjectClass->IsChildOf(AActor::StaticClass()))
@@ -557,16 +662,15 @@ void UEditorUtilitySubsystem::RunTaskCommand(const TArray<FString>& Params, UWor
 		FString TaskName = Params[0];
 		if (UClass* FoundClass = FindClassByName(TaskName))
 		{
-			TSubclassOf<UEditorUtilityTask> TaskToSpawn(FoundClass);
-			if (FoundClass == nullptr)
+			if (!FoundClass->IsChildOf<UEditorUtilityTask>())
 			{
 				UE_LOG(LogEditorUtilityBlueprint, Error, TEXT("Found Task: %s, but it's not a subclass of 'EditorUtilityTask'."), *FoundClass->GetName());
 				return;
 			}
 
-			UE_LOG(LogEditorUtilityBlueprint, Log, TEXT("Running task %s"), *TaskToSpawn->GetPathName());
+			UE_LOG(LogEditorUtilityBlueprint, Log, TEXT("Running task %s"), *FoundClass->GetPathName());
 
-			UEditorUtilityTask* NewTask = NewObject<UEditorUtilityTask>(this, *TaskToSpawn);
+			UEditorUtilityTask* NewTask = NewObject<UEditorUtilityTask>(this, FoundClass);
 			
 			//TODO Attempt to map XXX=YYY to properties on the task to make the tasks parameterizable
 
@@ -700,21 +804,18 @@ UClass* UEditorUtilitySubsystem::FindClassByName(const FString& RawTargetName)
 	return ResultClass;
 }
 
-UClass* UEditorUtilitySubsystem::FindBlueprintClass(const FString& TargetNameRaw)
+UClass* UEditorUtilitySubsystem::FindBlueprintClass(const FString& TargetName)
 {
-	FAssetRegistryModule& AssetRegistryModule = FModuleManager::LoadModuleChecked<FAssetRegistryModule>("AssetRegistry");
-	IAssetRegistry& AssetRegistry = AssetRegistryModule.Get();
+	IAssetRegistry& AssetRegistry = IAssetRegistry::GetChecked();
 	if (AssetRegistry.IsLoadingAssets())
 	{
 		AssetRegistry.SearchAllAssets(true);
 	}
 
-	FString TargetName = TargetNameRaw;
-	TargetName.RemoveFromEnd(TEXT("_C"), ESearchCase::CaseSensitive);
-
 	FARFilter Filter;
 	Filter.bRecursiveClasses = true;
 	Filter.ClassPaths.Add(UBlueprintCore::StaticClass()->GetClassPathName());
+	Filter.ClassPaths.Add(UBlueprintGeneratedClass::StaticClass()->GetClassPathName());
 
 	// We enumerate all assets to find any blueprints who inherit from native classes directly - or
 	// from other blueprints.
@@ -723,15 +824,15 @@ UClass* UEditorUtilitySubsystem::FindBlueprintClass(const FString& TargetNameRaw
 	{
 		if ((AssetData.AssetName.ToString() == TargetName) || (AssetData.GetObjectPathString() == TargetName))
 		{
-			if (UBlueprint* BP = Cast<UBlueprint>(AssetData.GetAsset()))
+			if (UObject* Asset = AssetData.GetAsset())
 			{
-				FoundClass = BP->GeneratedClass;
+				FoundClass = UE::EditorUtility::Private::GetRunnableClassForAsset(Asset);
 				return false;
 			}
 		}
 
 		return true;
-	});
+	}, UE::AssetRegistry::EEnumerateAssetsFlags::AllowUnfilteredArAssets);
 
 	return FoundClass;
 }
diff --git a/Engine/Source/Editor/Blutility/Public/EditorUtilitySubsystem.h b/Engine/Source/Editor/Blutility/Public/EditorUtilitySubsystem.h
--- a/Engine/Source/Editor/Blutility/Public/EditorUtilitySubsystem.h
+++ b/Engine/Source/Editor/Blutility/Public/EditorUtilitySubsystem.h
@@ -28,6 +28,7 @@ class IConsoleObject;
 class SDockTab;
 class SWindow;
 class UClass;
+class UBlueprintGeneratedClass;
 class UEditorUtilityTask;
 class UEditorUtilityWidget;
 class UEditorUtilityWidgetBlueprint;
@@ -145,6 +146,11 @@ public:
 	FOnEditorUtilityPIEEvent OnEndPIE;
 
 protected:
+	void HandleStartupAssets();
+	void HandleDynamicStartupAssets(TConstArrayView<FAssetData> InAssets);
+
+	void OnBlueprintGeneratedClassUnloaded(UBlueprintGeneratedClass* BPGC);
+
 	UEditorUtilityTask* GetActiveTask() { return ActiveTaskStack.Num() > 0 ? ActiveTaskStack[ActiveTaskStack.Num() - 1] : nullptr; };
 
 	void StartTask(UEditorUtilityTask* Task);
