UE-282993 - Asset/Actor Actions in Developer folders show for everyone

GetBlutilityClasses already had code to ensure that we don't list other devs' blutilities. However, it didn't use the correct path format, which caused FPaths::IsUnderDirectory to always return false;

Similarly, UEditorUtilitySubsystem::HandleStartupAssets never had a similar filter to begin with.

diff --git a/Engine/Source/Editor/Blutility/Private/BlutilityMenuExtensions.cpp b/Engine/Source/Editor/Blutility/Private/BlutilityMenuExtensions.cpp
--- a/Engine/Source/Editor/Blutility/Private/BlutilityMenuExtensions.cpp
+++ b/Engine/Source/Editor/Blutility/Private/BlutilityMenuExtensions.cpp
@@ -578,9 +578,10 @@ void FBlutilityMenuExtensions::GetBlutilityClasses(TArray<FAssetData>& OutAssets
 		if (Result.IsSet())
 		{
 			// If it's a menu extension blutility, don't include ones from other peoples developer folders, just your own.
-			if (FPaths::IsUnderDirectory(Asset.PackagePath.ToString(), FPaths::GameDevelopersDir()))
+			const FString& AssetPath = FPackageName::LongPackageNameToFilename(Asset.PackagePath.ToString());
+			if (FPaths::IsUnderDirectory(AssetPath, FPaths::GameDevelopersDir()))
 			{
-				if (!FPaths::IsUnderDirectory(Asset.PackagePath.ToString(), FPaths::GameUserDeveloperFolderName()))
+				if (!FPaths::IsUnderDirectory(AssetPath, FPaths::GameUserDeveloperDir()))
 				{
 					continue;
 				}
diff --git a/Engine/Source/Editor/Blutility/Private/EditorUtilitySubsystem.cpp b/Engine/Source/Editor/Blutility/Private/EditorUtilitySubsystem.cpp
--- a/Engine/Source/Editor/Blutility/Private/EditorUtilitySubsystem.cpp
+++ b/Engine/Source/Editor/Blutility/Private/EditorUtilitySubsystem.cpp
@@ -212,6 +212,15 @@ void UEditorUtilitySubsystem::HandleStartupAssets()
 
 		for (const FAssetData& StartupAsset : StartupAssets)
 		{
+			const FString& AssetPath = FPackageName::LongPackageNameToFilename(StartupAsset.PackagePath.ToString());
+			if (FPaths::IsUnderDirectory(AssetPath, FPaths::GameDevelopersDir()))
+			{
+				if (!FPaths::IsUnderDirectory(AssetPath, FPaths::GameUserDeveloperDir()))
+				{
+					continue;
+				}
+			}
+
 			UObject* Object = StartupAsset.GetAsset();
 			if (!Object || Object->IsUnreachable())
 			{
