Updated headers using UnrealCodeFixup to make sure dllstorage is on methods/staticvars instead of types. Part 4/n

diff --git a/Engine/Source/Editor/Blutility/Classes/ActorActionUtility.h b/Engine/Source/Editor/Blutility/Classes/ActorActionUtility.h
--- a/Engine/Source/Editor/Blutility/Classes/ActorActionUtility.h
+++ b/Engine/Source/Editor/Blutility/Classes/ActorActionUtility.h
@@ -7,13 +7,15 @@
 #include "IEditorUtilityExtension.h"
 #include "ActorActionUtility.generated.h"
 
+#define UE_API BLUTILITY_API
+
 /** 
  * Base class for all actor action-related utilities
  * Any functions/events that are exposed on derived classes that have the correct signature will be
  * included as menu options when right-clicking on a group of actors in the level editor.
  */
-UCLASS(Abstract, hideCategories=(Object), Blueprintable)
-class BLUTILITY_API UActorActionUtility : public UEditorUtilityObject, public IEditorUtilityExtension
+UCLASS(MinimalAPI, Abstract, hideCategories=(Object), Blueprintable)
+class UActorActionUtility : public UEditorUtilityObject, public IEditorUtilityExtension
 {
 	GENERATED_BODY()
 
@@ -21,7 +23,7 @@ public:
 	/** Return the class that this actor action supports. Leave this blank to support all actor classes. */
 	UE_DEPRECATED(5.2, "GetSupportedClasses() instead, but ideally you're not requesting this directly and are instead using the FAssetActionUtilityPrototype to wrap access to an unload utility asset.")
 	UFUNCTION(BlueprintPure, BlueprintImplementableEvent, Category="Assets", meta=(DeprecatedFunction, DeprecationMessage="If you were just returning a single class add it to the SupportedClasses array (you can find it listed in the Class Defaults).  If you were doing complex logic to simulate having multiple classes act as filters, add them to the SupportedClasses array.  If you were doing 'other' logic, you'll need to do that upon action execution."))
-	UClass* GetSupportedClass() const;
+	UE_API UClass* GetSupportedClass() const;
 	
 	/**
      * Gets the statically determined supported classes, these classes are used as a first pass filter when determining
@@ -32,9 +34,9 @@ public:
 
 public:
     // Begin UObject
-	virtual void GetAssetRegistryTags(FAssetRegistryTagsContext Context) const override;
+	UE_API virtual void GetAssetRegistryTags(FAssetRegistryTagsContext Context) const override;
 	UE_DEPRECATED(5.4, "Implement the version that takes FAssetRegistryTagsContext instead.")
-	virtual void GetAssetRegistryTags(TArray<FAssetRegistryTag>& OutTags) const override;
+	UE_API virtual void GetAssetRegistryTags(TArray<FAssetRegistryTag>& OutTags) const override;
 	// End UObject
 	
 private:
@@ -43,4 +45,6 @@ private:
 	 */
 	UPROPERTY(EditDefaultsOnly, Category="Assets", meta=(AllowAbstract))
 	TArray<TSoftClassPtr<UObject>> SupportedClasses;
-};
\ No newline at end of file
+};
+
+#undef UE_API
diff --git a/Engine/Source/Editor/Blutility/Classes/AssetActionUtility.h b/Engine/Source/Editor/Blutility/Classes/AssetActionUtility.h
--- a/Engine/Source/Editor/Blutility/Classes/AssetActionUtility.h
+++ b/Engine/Source/Editor/Blutility/Classes/AssetActionUtility.h
@@ -8,6 +8,8 @@
 #include "IEditorUtilityExtension.h"
 #include "AssetActionUtility.generated.h"
 
+#define UE_API BLUTILITY_API
+
 namespace AssetActionUtilityTags
 {
 	extern const FName BlutilityTagVersion;
@@ -54,8 +56,8 @@ struct FAssetActionSupportCondition
  * Any functions/events that are exposed on derived classes that have the correct signature will be
  * included as menu options when right-clicking on a group of assets in the content browser.
  */
-UCLASS(Abstract, hideCategories=(Object), Blueprintable)
-class BLUTILITY_API UAssetActionUtility : public UEditorUtilityObject, public IEditorUtilityExtension
+UCLASS(MinimalAPI, Abstract, hideCategories=(Object), Blueprintable)
+class UAssetActionUtility : public UEditorUtilityObject, public IEditorUtilityExtension
 {
 	GENERATED_BODY()
 
@@ -66,7 +68,7 @@ public:
 	 */
 	UE_DEPRECATED(5.2, "GetSupportedClasses() instead, but ideally you're not requesting this directly and are instead using the FAssetActionUtilityPrototype to wrap access to an unload utility asset.")
 	UFUNCTION(BlueprintPure, BlueprintImplementableEvent, Category="Assets", meta=(DeprecatedFunction, DeprecationMessage="If you were just returning a single class add it to the SupportedClasses array (you can find it listed in the Class Defaults).  If you were doing complex logic to simulate having multiple classes act as filters, add them to the SupportedClasses array.  If you were doing 'other' logic, you'll need to do that upon action execution."))
-	UClass* GetSupportedClass() const;
+	UE_API UClass* GetSupportedClass() const;
 
 	/**
 	 * Returns whether or not this action is designed to work specifically on Blueprints (true) or on all assets (false).
@@ -74,7 +76,7 @@ public:
 	 * @note Returns the value of bIsActionForBlueprints by default.
 	 */
 	UFUNCTION(BlueprintPure, BlueprintNativeEvent, Category="Assets")
-	bool IsActionForBlueprints() const;
+	UE_API bool IsActionForBlueprints() const;
 
 	/**
 	 * Gets the statically determined supported classes, these classes are used as a first pass filter when determining
@@ -85,9 +87,9 @@ public:
 
 public:
 	// Begin UObject
-	virtual void GetAssetRegistryTags(FAssetRegistryTagsContext Context) const override;
+	UE_API virtual void GetAssetRegistryTags(FAssetRegistryTagsContext Context) const override;
 	UE_DEPRECATED(5.4, "Implement the version that takes FAssetRegistryTagsContext instead.")
-	virtual void GetAssetRegistryTags(TArray<FAssetRegistryTag>& OutTags) const override;
+	UE_API virtual void GetAssetRegistryTags(TArray<FAssetRegistryTag>& OutTags) const override;
 	// End UObject
 
 protected:
@@ -114,3 +116,5 @@ protected:
 	UPROPERTY(EditDefaultsOnly, Category="Asset Support")
 	TArray<FAssetActionSupportCondition> SupportedConditions;
 };
+
+#undef UE_API
diff --git a/Engine/Source/Editor/Blutility/Classes/EditorUtilityActor.h b/Engine/Source/Editor/Blutility/Classes/EditorUtilityActor.h
--- a/Engine/Source/Editor/Blutility/Classes/EditorUtilityActor.h
+++ b/Engine/Source/Editor/Blutility/Classes/EditorUtilityActor.h
@@ -17,23 +17,25 @@
 
 #include "EditorUtilityActor.generated.h"
 
+#define UE_API BLUTILITY_API
+
 class UInputComponent;
 class UObject;
 struct FFrame;
 struct FPropertyChangedEvent;
 
-UCLASS(Abstract, Blueprintable, meta = (ShowWorldContextPin))
-class BLUTILITY_API AEditorUtilityActor : public AActor
+UCLASS(MinimalAPI, Abstract, Blueprintable, meta = (ShowWorldContextPin))
+class AEditorUtilityActor : public AActor
 {
 	GENERATED_UCLASS_BODY()
 
 	// Standard function to execute
 	UFUNCTION(BlueprintCallable, BlueprintImplementableEvent, Category = "Editor")
-	void Run();
+	UE_API void Run();
 
-	virtual void OnConstruction(const FTransform& Transform) override;
+	UE_API virtual void OnConstruction(const FTransform& Transform) override;
 #if WITH_EDITOR
-	virtual void PostEditChangeProperty(FPropertyChangedEvent& PropertyChangedEvent) override;
+	UE_API virtual void PostEditChangeProperty(FPropertyChangedEvent& PropertyChangedEvent) override;
 #endif
 
 	/** Returns the current InputComponent on this utility actor. This will be NULL unless bReceivesEditorInput is set to true. */
@@ -44,7 +46,7 @@ class BLUTILITY_API AEditorUtilityActor : public AActor
 	}
 	
 	UFUNCTION(BlueprintSetter, Category = "Input|Editor")
-	void SetReceivesEditorInput(bool bInValue);
+	UE_API void SetReceivesEditorInput(bool bInValue);
 	
 	UFUNCTION(BlueprintGetter, BlueprintPure, Category = "Input|Editor")
 	bool GetReceivesEditorInput() const
@@ -55,10 +57,10 @@ class BLUTILITY_API AEditorUtilityActor : public AActor
 private:
 
 	/** Creates the EditorOnlyInputComponent if it does not already exist and registers all subobject callbacks to it */
-	void CreateEditorInput();
+	UE_API void CreateEditorInput();
 
 	/** Removes the EditorOnlyInputComponent from this utility actor */
-	void RemoveEditorInput();
+	UE_API void RemoveEditorInput();
 	
 	UPROPERTY(Transient, DuplicateTransient)
 	TObjectPtr<UInputComponent> EditorOnlyInputComponent;
@@ -67,3 +69,5 @@ private:
 	UPROPERTY(EditAnywhere, Category = "Input|Editor", BlueprintSetter = SetReceivesEditorInput, BlueprintGetter = GetReceivesEditorInput)
 	bool bReceivesEditorInput = false;
 };
+
+#undef UE_API
diff --git a/Engine/Source/Editor/Blutility/Classes/EditorUtilityActorComponent.h b/Engine/Source/Editor/Blutility/Classes/EditorUtilityActorComponent.h
--- a/Engine/Source/Editor/Blutility/Classes/EditorUtilityActorComponent.h
+++ b/Engine/Source/Editor/Blutility/Classes/EditorUtilityActorComponent.h
@@ -15,8 +15,8 @@
 class UObject;
 
 
-UCLASS(Abstract, Blueprintable, meta = (ShowWorldContextPin))
-class BLUTILITY_API UEditorUtilityActorComponent : public UActorComponent
+UCLASS(MinimalAPI, Abstract, Blueprintable, meta = (ShowWorldContextPin))
+class UEditorUtilityActorComponent : public UActorComponent
 {
 	GENERATED_UCLASS_BODY()
 
diff --git a/Engine/Source/Editor/Blutility/Classes/EditorUtilityBlueprint.h b/Engine/Source/Editor/Blutility/Classes/EditorUtilityBlueprint.h
--- a/Engine/Source/Editor/Blutility/Classes/EditorUtilityBlueprint.h
+++ b/Engine/Source/Editor/Blutility/Classes/EditorUtilityBlueprint.h
@@ -12,15 +12,19 @@
 
 #include "EditorUtilityBlueprint.generated.h"
 
+#define UE_API BLUTILITY_API
+
 class UObject;
 
-UCLASS()
-class BLUTILITY_API UEditorUtilityBlueprint : public UBlueprint
+UCLASS(MinimalAPI)
+class UEditorUtilityBlueprint : public UBlueprint
 {
 	GENERATED_UCLASS_BODY()
 
 	// UBlueprint interface
-	virtual bool SupportedByDefaultBlueprintFactory() const override;
-	virtual bool AlwaysCompileOnLoad() const override;
+	UE_API virtual bool SupportedByDefaultBlueprintFactory() const override;
+	UE_API virtual bool AlwaysCompileOnLoad() const override;
 	// End of UBlueprint interface
 };
+
+#undef UE_API
diff --git a/Engine/Source/Editor/Blutility/Classes/EditorUtilityBlueprintFactory.h b/Engine/Source/Editor/Blutility/Classes/EditorUtilityBlueprintFactory.h
--- a/Engine/Source/Editor/Blutility/Classes/EditorUtilityBlueprintFactory.h
+++ b/Engine/Source/Editor/Blutility/Classes/EditorUtilityBlueprintFactory.h
@@ -16,13 +16,15 @@
 
 #include "EditorUtilityBlueprintFactory.generated.h"
 
+#define UE_API BLUTILITY_API
+
 class FFeedbackContext;
 class SWindow;
 class UClass;
 class UObject;
 
-UCLASS()
-class BLUTILITY_API UEditorUtilityBlueprintFactory : public UFactory
+UCLASS(MinimalAPI)
+class UEditorUtilityBlueprintFactory : public UFactory
 {
 	GENERATED_UCLASS_BODY()
 
@@ -31,15 +33,17 @@ class BLUTILITY_API UEditorUtilityBlueprintFactory : public UFactory
 	TSubclassOf<class UObject> ParentClass;
 
 	// UFactory interface
-	virtual bool ConfigureProperties() override;
-	virtual UObject* FactoryCreateNew(UClass* Class, UObject* InParent, FName Name, EObjectFlags Flags, UObject* Context, FFeedbackContext* Warn) override;
-	virtual bool CanCreateNew() const override;
+	UE_API virtual bool ConfigureProperties() override;
+	UE_API virtual UObject* FactoryCreateNew(UClass* Class, UObject* InParent, FName Name, EObjectFlags Flags, UObject* Context, FFeedbackContext* Warn) override;
+	UE_API virtual bool CanCreateNew() const override;
 	// End of UFactory interface
 
 	/** Handler for when a class is picked in the class picker */
-	void OnClassPicked(UClass* InChosenClass);
+	UE_API void OnClassPicked(UClass* InChosenClass);
 
 protected:
 	/** A pointer to the window that is asking the user to select a parent class */
 	TWeakPtr<SWindow> PickerWindow;
 };
+
+#undef UE_API
diff --git a/Engine/Source/Editor/Blutility/Classes/EditorUtilityCamera.h b/Engine/Source/Editor/Blutility/Classes/EditorUtilityCamera.h
--- a/Engine/Source/Editor/Blutility/Classes/EditorUtilityCamera.h
+++ b/Engine/Source/Editor/Blutility/Classes/EditorUtilityCamera.h
@@ -15,10 +15,10 @@
 class UObject;
 
 
-UCLASS(Abstract, Blueprintable, meta = (ShowWorldContextPin))
-class BLUTILITY_API AEditorUtilityCamera : public ACameraActor
+UCLASS(MinimalAPI, Abstract, Blueprintable, meta = (ShowWorldContextPin))
+class AEditorUtilityCamera : public ACameraActor
 {
 	GENERATED_UCLASS_BODY()
 
 
-};
\ No newline at end of file
+};
diff --git a/Engine/Source/Editor/Blutility/Classes/EditorUtilityObject.h b/Engine/Source/Editor/Blutility/Classes/EditorUtilityObject.h
--- a/Engine/Source/Editor/Blutility/Classes/EditorUtilityObject.h
+++ b/Engine/Source/Editor/Blutility/Classes/EditorUtilityObject.h
@@ -13,15 +13,17 @@
 
 #include "EditorUtilityObject.generated.h"
 
+#define UE_API BLUTILITY_API
 
-UCLASS(Abstract, Blueprintable, meta = (ShowWorldContextPin))
-class BLUTILITY_API UEditorUtilityObject : public UObject, public IAssetRegistryTagProviderInterface
+
+UCLASS(MinimalAPI, Abstract, Blueprintable, meta = (ShowWorldContextPin))
+class UEditorUtilityObject : public UObject, public IAssetRegistryTagProviderInterface
 {
 	GENERATED_UCLASS_BODY()
 
 	// Standard function to execute
 	UFUNCTION(BlueprintCallable, BlueprintImplementableEvent, Category = "Editor")
-	void Run();
+	UE_API void Run();
 
 	// we use this to call tool started delegate
 	virtual void PostLoad() override;
@@ -38,3 +40,5 @@ protected:
 	UPROPERTY(Category=Settings, EditDefaultsOnly, AssetRegistrySearchable, DisplayName="Run on Start-up")
 	bool bRunEditorUtilityOnStartup = false;
 };
+
+#undef UE_API
diff --git a/Engine/Source/Editor/Blutility/Classes/EditorUtilityTask.h b/Engine/Source/Editor/Blutility/Classes/EditorUtilityTask.h
--- a/Engine/Source/Editor/Blutility/Classes/EditorUtilityTask.h
+++ b/Engine/Source/Editor/Blutility/Classes/EditorUtilityTask.h
@@ -15,6 +15,8 @@
 
 #include "EditorUtilityTask.generated.h"
 
+#define UE_API BLUTILITY_API
+
 class FAsyncTaskNotification;
 class FText;
 class UEditorUtilitySubsystem;
@@ -26,8 +28,8 @@ DECLARE_MULTICAST_DELEGATE_OneParam(FOnEditorUtilityTaskDelegate, UEditorUtility
 /**
  * 
  */
-UCLASS(Abstract, Blueprintable, meta = (ShowWorldContextPin))
-class BLUTILITY_API UEditorUtilityTask : public UObject, public IAssetRegistryTagProviderInterface
+UCLASS(MinimalAPI, Abstract, Blueprintable, meta = (ShowWorldContextPin))
+class UEditorUtilityTask : public UObject, public IAssetRegistryTagProviderInterface
 {
 	GENERATED_BODY()
 
@@ -35,10 +37,10 @@ public:
 	FOnEditorUtilityTaskDelegate OnFinished;
 
 public:
-	UEditorUtilityTask();
+	UE_API UEditorUtilityTask();
 
 	UFUNCTION()
-	void Run();
+	UE_API void Run();
 
 	//~ Begin IAssetRegistryTagProviderInterface interface
 	virtual bool ShouldAddCDOTagsToBlueprintClass() const override
@@ -47,41 +49,41 @@ public:
 	}
 	//~ End IAssetRegistryTagProviderInterface interface
 
-	virtual UWorld* GetWorld() const override;
+	UE_API virtual UWorld* GetWorld() const override;
 
 public:
 	UFUNCTION(BlueprintCallable, Category=Task)
-	void FinishExecutingTask(const bool bSuccess = true);
+	UE_API void FinishExecutingTask(const bool bSuccess = true);
 
 	UFUNCTION(BlueprintCallable, Category = Task)
-	void SetTaskNotificationText(const FText& Text);
+	UE_API void SetTaskNotificationText(const FText& Text);
 
 	// Calls CancelRequested() and ReceiveCancelRequested()
-	void RequestCancel();
+	UE_API void RequestCancel();
 
 	UFUNCTION(BlueprintCallable, Category = Task)
-	bool WasCancelRequested() const;
+	UE_API bool WasCancelRequested() const;
 
 protected:
 	virtual void BeginExecution() {}
 	virtual void CancelRequested() {}
-	virtual FText GetTaskTitle() const;
+	UE_API virtual FText GetTaskTitle() const;
 
 	UFUNCTION(BlueprintImplementableEvent, Category=Task, meta=(DisplayName="BeginExecution"))
-	void ReceiveBeginExecution();
+	UE_API void ReceiveBeginExecution();
 
 	UFUNCTION(BlueprintImplementableEvent, Category=Task, meta=(DisplayName="CancelRequested"))
-	void ReceiveCancelRequested();
+	UE_API void ReceiveCancelRequested();
 
 	UFUNCTION(BlueprintImplementableEvent, Category=Task)
-	FText GetTaskTitleOverride() const;
+	UE_API FText GetTaskTitleOverride() const;
 
 private:
 	// Calls GetTaskTitle() and GetTaskTitleOverride()
-	void CreateNotification();
+	UE_API void CreateNotification();
 
 	// Calls BeginExecution() and ReceiveBeginExecution()
-	void StartExecutingTask();
+	UE_API void StartExecutingTask();
 
 protected:
 	/** Run this editor utility on start-up (after asset discovery)? */
@@ -104,3 +106,5 @@ private:
 
 	friend UEditorUtilitySubsystem;
 };
+
+#undef UE_API
diff --git a/Engine/Source/Editor/Blutility/Classes/EditorUtilityToolMenu.h b/Engine/Source/Editor/Blutility/Classes/EditorUtilityToolMenu.h
--- a/Engine/Source/Editor/Blutility/Classes/EditorUtilityToolMenu.h
+++ b/Engine/Source/Editor/Blutility/Classes/EditorUtilityToolMenu.h
@@ -10,8 +10,8 @@
 #include "EditorUtilityToolMenu.generated.h"
 
 
-UCLASS(Blueprintable, abstract)
-class BLUTILITY_API UEditorUtilityToolMenuEntry : public UToolMenuEntryScript, public IAssetRegistryTagProviderInterface
+UCLASS(MinimalAPI, Blueprintable, abstract)
+class UEditorUtilityToolMenuEntry : public UToolMenuEntryScript, public IAssetRegistryTagProviderInterface
 {
 	GENERATED_BODY()
 
@@ -28,8 +28,8 @@ protected:
 	bool bRunEditorUtilityOnStartup = false;
 };
 
-UCLASS(Blueprintable, abstract)
-class BLUTILITY_API UEditorUtilityToolMenuSection : public UToolMenuSectionDynamic, public IAssetRegistryTagProviderInterface
+UCLASS(MinimalAPI, Blueprintable, abstract)
+class UEditorUtilityToolMenuSection : public UToolMenuSectionDynamic, public IAssetRegistryTagProviderInterface
 {
 	GENERATED_BODY()
 
diff --git a/Engine/Source/Editor/Blutility/Classes/EditorUtilityWidget.h b/Engine/Source/Editor/Blutility/Classes/EditorUtilityWidget.h
--- a/Engine/Source/Editor/Blutility/Classes/EditorUtilityWidget.h
+++ b/Engine/Source/Editor/Blutility/Classes/EditorUtilityWidget.h
@@ -17,19 +17,21 @@
 
 #include "EditorUtilityWidget.generated.h"
 
+#define UE_API BLUTILITY_API
+
 class AActor;
 class UEditorPerProjectUserSettings;
 class UObject;
 
-UCLASS(Abstract, meta = (ShowWorldContextPin), config = Editor)
-class BLUTILITY_API UEditorUtilityWidget : public UUserWidget, public IAssetRegistryTagProviderInterface
+UCLASS(MinimalAPI, Abstract, meta = (ShowWorldContextPin), config = Editor)
+class UEditorUtilityWidget : public UUserWidget, public IAssetRegistryTagProviderInterface
 {
 	GENERATED_BODY()
 
 public:
 	// The default action called when the widget is invoked if bAutoRunDefaultAction=true (it is never called otherwise)
 	UFUNCTION(BlueprintCallable, BlueprintImplementableEvent, Category = "Editor")
-	void Run();
+	UE_API void Run();
 
 	//~ Begin IAssetRegistryTagProviderInterface interface
 	virtual bool ShouldAddCDOTagsToBlueprintClass() const override
@@ -39,10 +41,10 @@ public:
 	//~ End IAssetRegistryTagProviderInterface interface
 
 	UFUNCTION(BlueprintCallable, BlueprintPure = false, Category = "Editor Utility Widget")
-	UWidget* FindChildWidgetByName(FName WidgetName) const;
+	UE_API UWidget* FindChildWidgetByName(FName WidgetName) const;
 
 	// Run the default action
-	void ExecuteDefaultAction();
+	UE_API void ExecuteDefaultAction();
 
 	bool ShouldAlwaysReregisterWithWindowsMenu() const
 	{
@@ -86,3 +88,5 @@ protected:
 	UPROPERTY(Category=Settings, EditDefaultsOnly, AssetRegistrySearchable, DisplayName="Run on Start-up")
 	bool bRunEditorUtilityOnStartup = false;
 };
+
+#undef UE_API
diff --git a/Engine/Source/Editor/Blutility/Classes/EditorUtilityWidgetBlueprint.h b/Engine/Source/Editor/Blutility/Classes/EditorUtilityWidgetBlueprint.h
--- a/Engine/Source/Editor/Blutility/Classes/EditorUtilityWidgetBlueprint.h
+++ b/Engine/Source/Editor/Blutility/Classes/EditorUtilityWidgetBlueprint.h
@@ -21,6 +21,8 @@
 
 #include "EditorUtilityWidgetBlueprint.generated.h"
 
+#define UE_API BLUTILITY_API
+
 class FSpawnTabArgs;
 class SDockTab;
 class SWidget;
@@ -33,26 +35,26 @@ class UWorld;
 enum class EAssetEditorCloseReason : uint8;
 enum class EMapChangeType : uint8;
 
-UCLASS()
-class BLUTILITY_API UEditorUtilityWidgetBlueprint : public UWidgetBlueprint
+UCLASS(MinimalAPI)
+class UEditorUtilityWidgetBlueprint : public UWidgetBlueprint
 {
 	GENERATED_BODY()
 
 public:
-	virtual void BeginDestroy() override;
+	UE_API virtual void BeginDestroy() override;
 
-	TSharedRef<SDockTab> SpawnEditorUITab(const FSpawnTabArgs& SpawnTabArgs);
+	UE_API TSharedRef<SDockTab> SpawnEditorUITab(const FSpawnTabArgs& SpawnTabArgs);
 
 	/** Creates the slate widget from the UMG widget */
-	TSharedRef<SWidget> CreateUtilityWidget();
+	UE_API TSharedRef<SWidget> CreateUtilityWidget();
 
 	/** Recreate the tab's content on recompile */
-	void RegenerateCreatedTab(UBlueprint* RecompiledBlueprint);
+	UE_API void RegenerateCreatedTab(UBlueprint* RecompiledBlueprint);
 	
-	void UpdateRespawnListIfNeeded(TSharedRef<SDockTab> TabBeingClosed);
+	UE_API void UpdateRespawnListIfNeeded(TSharedRef<SDockTab> TabBeingClosed);
 
 	// UBlueprint interface
-	virtual void GetReparentingRules(TSet< const UClass* >& AllowedChildrenOfClasses, TSet< const UClass* >& DisallowedChildrenOfClasses) const override;
+	UE_API virtual void GetReparentingRules(TSet< const UClass* >& AllowedChildrenOfClasses, TSet< const UClass* >& DisallowedChildrenOfClasses) const override;
 
 	virtual bool AllowEditorWidget() const override { return true; }
 
@@ -72,23 +74,23 @@ public:
 	}
 
 	/** Returns the default desired tab display name that was specified for this widget */
-	FText GetTabDisplayName() const;
+	UE_API FText GetTabDisplayName() const;
 
 	bool ShouldSpawnAsNomadTab() const
 	{
 		return bSpawnAsNomadTab;
 	}
 
-	virtual UWidgetEditingProjectSettings* GetRelevantSettings() override;
-	virtual const UWidgetEditingProjectSettings* GetRelevantSettings() const override;
+	UE_API virtual UWidgetEditingProjectSettings* GetRelevantSettings() override;
+	UE_API virtual const UWidgetEditingProjectSettings* GetRelevantSettings() const override;
 
 public:
-	static void MarkTransientRecursive(UEditorUtilityWidget* UtilityWidget);
+	static UE_API void MarkTransientRecursive(UEditorUtilityWidget* UtilityWidget);
 
 private:
-	bool IsWidgetEnabled() const;
+	UE_API bool IsWidgetEnabled() const;
 
-	void ChangeTabWorld(UWorld* World, EMapChangeType MapChangeType);
+	UE_API void ChangeTabWorld(UWorld* World, EMapChangeType MapChangeType);
 
 private:
 	// Should the widget be enabled when running PIE
@@ -111,3 +113,5 @@ private:
 	TObjectPtr<UEditorUtilityWidget> CreatedUMGWidget;
 
 };
+
+#undef UE_API
diff --git a/Engine/Source/Editor/Blutility/Classes/EditorUtilityWidgetComponents.h b/Engine/Source/Editor/Blutility/Classes/EditorUtilityWidgetComponents.h
--- a/Engine/Source/Editor/Blutility/Classes/EditorUtilityWidgetComponents.h
+++ b/Engine/Source/Editor/Blutility/Classes/EditorUtilityWidgetComponents.h
@@ -32,119 +32,123 @@
 
 #include "EditorUtilityWidgetComponents.generated.h"
 
-UCLASS()
-class BLUTILITY_API UEditorUtilityButton : public UButton
+#define UE_API BLUTILITY_API
+
+UCLASS(MinimalAPI)
+class UEditorUtilityButton : public UButton
 {
 	GENERATED_UCLASS_BODY()
 };
 
-UCLASS()
-class BLUTILITY_API UEditorUtilityCheckBox : public UCheckBox
+UCLASS(MinimalAPI)
+class UEditorUtilityCheckBox : public UCheckBox
 {
 	GENERATED_UCLASS_BODY()
 };
 
-UCLASS()
-class BLUTILITY_API UEditorUtilityCircularThrobber : public UCircularThrobber
+UCLASS(MinimalAPI)
+class UEditorUtilityCircularThrobber : public UCircularThrobber
 {
 	GENERATED_UCLASS_BODY()
 };
 
-UCLASS()
-class BLUTILITY_API UEditorUtilityComboBoxKey : public UComboBoxKey
+UCLASS(MinimalAPI)
+class UEditorUtilityComboBoxKey : public UComboBoxKey
 {
 	GENERATED_BODY()
 	
 public:
-	UEditorUtilityComboBoxKey();
+	UE_API UEditorUtilityComboBoxKey();
 };
 
-UCLASS()
-class BLUTILITY_API UEditorUtilityComboBoxString : public UComboBoxString
+UCLASS(MinimalAPI)
+class UEditorUtilityComboBoxString : public UComboBoxString
 {
 	GENERATED_UCLASS_BODY()
 };
 
-UCLASS()
-class BLUTILITY_API UEditorUtilityEditableText : public UEditableText
+UCLASS(MinimalAPI)
+class UEditorUtilityEditableText : public UEditableText
 {
 	GENERATED_UCLASS_BODY()
 };
 
-UCLASS()
-class BLUTILITY_API UEditorUtilityEditableTextBox : public UEditableTextBox
+UCLASS(MinimalAPI)
+class UEditorUtilityEditableTextBox : public UEditableTextBox
 {
 	GENERATED_UCLASS_BODY()
 };
 
-UCLASS()
-class BLUTILITY_API UEditorUtilityExpandableArea : public UExpandableArea
+UCLASS(MinimalAPI)
+class UEditorUtilityExpandableArea : public UExpandableArea
 {
 	GENERATED_UCLASS_BODY()
 };
 
-UCLASS()
-class BLUTILITY_API UEditorUtilityInputKeySelector : public UInputKeySelector
+UCLASS(MinimalAPI)
+class UEditorUtilityInputKeySelector : public UInputKeySelector
 {
 	GENERATED_UCLASS_BODY()
 };
 
-UCLASS()
-class BLUTILITY_API UEditorUtilityListView : public UListView
+UCLASS(MinimalAPI)
+class UEditorUtilityListView : public UListView
 {
 	GENERATED_UCLASS_BODY()
 };
 
-UCLASS()
-class BLUTILITY_API UEditorUtilityMultiLineEditableText : public UMultiLineEditableText
+UCLASS(MinimalAPI)
+class UEditorUtilityMultiLineEditableText : public UMultiLineEditableText
 {
 	GENERATED_UCLASS_BODY()
 };
 
-UCLASS()
-class BLUTILITY_API UEditorUtilityMultiLineEditableTextBox : public UMultiLineEditableTextBox
+UCLASS(MinimalAPI)
+class UEditorUtilityMultiLineEditableTextBox : public UMultiLineEditableTextBox
 {
 	GENERATED_UCLASS_BODY()
 };
 
-UCLASS()
-class BLUTILITY_API UEditorUtilityProgressBar : public UProgressBar
+UCLASS(MinimalAPI)
+class UEditorUtilityProgressBar : public UProgressBar
 {
 	GENERATED_UCLASS_BODY()
 };
 
-UCLASS()
-class BLUTILITY_API UEditorUtilityScrollBar : public UScrollBar
+UCLASS(MinimalAPI)
+class UEditorUtilityScrollBar : public UScrollBar
 {
 	GENERATED_UCLASS_BODY()
 };
 
-UCLASS()
-class BLUTILITY_API UEditorUtilityScrollBox : public UScrollBox
+UCLASS(MinimalAPI)
+class UEditorUtilityScrollBox : public UScrollBox
 {
 	GENERATED_UCLASS_BODY()
 };
 
-UCLASS()
-class BLUTILITY_API UEditorUtilitySlider : public USlider
+UCLASS(MinimalAPI)
+class UEditorUtilitySlider : public USlider
 {
 	GENERATED_UCLASS_BODY()
 };
 
-UCLASS()
-class BLUTILITY_API UEditorUtilitySpinBox : public USpinBox
+UCLASS(MinimalAPI)
+class UEditorUtilitySpinBox : public USpinBox
 {
 	GENERATED_UCLASS_BODY()
 };
 
-UCLASS()
-class BLUTILITY_API UEditorUtilityThrobber : public UThrobber
+UCLASS(MinimalAPI)
+class UEditorUtilityThrobber : public UThrobber
 {
 	GENERATED_UCLASS_BODY()
 };
 
-UCLASS()
-class BLUTILITY_API UEditorUtilityTreeView : public UTreeView
+UCLASS(MinimalAPI)
+class UEditorUtilityTreeView : public UTreeView
 {
 	GENERATED_UCLASS_BODY()
 };
+
+#undef UE_API
diff --git a/Engine/Source/Editor/Blutility/Classes/GlobalEditorUtilityBase.h b/Engine/Source/Editor/Blutility/Classes/GlobalEditorUtilityBase.h
--- a/Engine/Source/Editor/Blutility/Classes/GlobalEditorUtilityBase.h
+++ b/Engine/Source/Editor/Blutility/Classes/GlobalEditorUtilityBase.h
@@ -19,6 +19,8 @@
 
 #include "GlobalEditorUtilityBase.generated.h"
 
+#define UE_API BLUTILITY_API
+
 class AActor;
 class UEditorPerProjectUserSettings;
 struct FFrame;
@@ -27,8 +29,8 @@ DECLARE_DYNAMIC_MULTICAST_DELEGATE_TwoParams( FForEachActorIteratorSignature, cl
 DECLARE_DYNAMIC_MULTICAST_DELEGATE_TwoParams( FForEachAssetIteratorSignature, class UObject*, Asset, int32, Index );
 
 
-UCLASS(Abstract, hideCategories=(Object), Blueprintable, Deprecated)
-class BLUTILITY_API UDEPRECATED_GlobalEditorUtilityBase : public UObject
+UCLASS(MinimalAPI, Abstract, hideCategories=(Object), Blueprintable, Deprecated)
+class UDEPRECATED_GlobalEditorUtilityBase : public UObject
 {
 	GENERATED_UCLASS_BODY()
 
@@ -39,10 +41,10 @@ class BLUTILITY_API UDEPRECATED_GlobalEditorUtilityBase : public UObject
 	bool bDirtiedSelectionSet;
 
 	/** UObject interface */
-	virtual UWorld* GetWorld() const override;
+	UE_API virtual UWorld* GetWorld() const override;
 
 	UFUNCTION(BlueprintCallable, Category="Development|Editor")
-	TArray<AActor*> GetSelectionSet();
+	UE_API TArray<AActor*> GetSelectionSet();
 
 	/**
 	 * Attempts to find the actor specified by PathToActor in the current editor world
@@ -50,7 +52,7 @@ class BLUTILITY_API UDEPRECATED_GlobalEditorUtilityBase : public UObject
 	 * @return	A reference to the actor, or none if it wasn't found
 	 */
 	UFUNCTION(BlueprintPure, Category = "Development|Editor")
-	AActor* GetActorReference(FString PathToActor);
+	UE_API AActor* GetActorReference(FString PathToActor);
 
 	////////////////////////////
 
@@ -60,13 +62,13 @@ class BLUTILITY_API UDEPRECATED_GlobalEditorUtilityBase : public UObject
 
 	// The default action called when the blutility is invoked if bAutoRunDefaultAction=true (it is never called otherwise)
 	UFUNCTION(BlueprintImplementableEvent)
-	void OnDefaultActionClicked();
+	UE_API void OnDefaultActionClicked();
 
 	////////////////////////////
 
 	// Calls OnEachSelectedActor for each selected actor
 	UFUNCTION(BlueprintCallable, Category="Development|Editor")
-	void ForEachSelectedActor();
+	UE_API void ForEachSelectedActor();
 
 	// The method called for each selected actor when ForEachSelectedActor is called
 	UPROPERTY(BlueprintAssignable)
@@ -76,7 +78,7 @@ class BLUTILITY_API UDEPRECATED_GlobalEditorUtilityBase : public UObject
 
 	// Calls OnEachSelectedAsset for each selected asset
 	UFUNCTION(BlueprintCallable, Category="Development|Editor")
-	void ForEachSelectedAsset();
+	UE_API void ForEachSelectedAsset();
 
 	// The method called for each selected asset when ForEachSelectedAsset is called
 	UPROPERTY(BlueprintAssignable)
@@ -84,37 +86,39 @@ class BLUTILITY_API UDEPRECATED_GlobalEditorUtilityBase : public UObject
 
 	// Gets the set of currently selected assets
 	UFUNCTION(BlueprintCallable, Category="Development|Editor")
-	TArray<UObject*> GetSelectedAssets();
+	UE_API TArray<UObject*> GetSelectedAssets();
 
 	///////////////////////////
 
 	UFUNCTION(BlueprintPure, Category="Development|Editor")
-	UEditorPerProjectUserSettings* GetEditorUserSettings();
+	UE_API UEditorPerProjectUserSettings* GetEditorUserSettings();
 
 	// Remove all actors from the selection set
 	UFUNCTION(BlueprintCallable, Category="Development|Editor")
-	void ClearActorSelectionSet();
+	UE_API void ClearActorSelectionSet();
 
 	// Selects nothing in the editor (another way to clear the selection)
 	UFUNCTION(BlueprintCallable, Category = "Development|Editor")
-	void SelectNothing();
+	UE_API void SelectNothing();
 
 	// Set the selection state for the selected actor
 	UFUNCTION(BlueprintCallable, Category="Development|Editor")
-	void SetActorSelectionState(AActor* Actor, bool bShouldBeSelected);
+	UE_API void SetActorSelectionState(AActor* Actor, bool bShouldBeSelected);
 
 	UFUNCTION(BlueprintCallable, Category="Development|Editor")
-	void GetSelectionBounds(FVector& Origin, FVector& BoxExtent, float& SphereRadius);
+	UE_API void GetSelectionBounds(FVector& Origin, FVector& BoxExtent, float& SphereRadius);
 
 	// Renames an asset (cannot move folders)
 	UFUNCTION(BlueprintCallable, Category = "Development|Editor")
-	void RenameAsset(UObject* Asset, const FString& NewName);
+	UE_API void RenameAsset(UObject* Asset, const FString& NewName);
 
 	////////////////////////////
 
 	// Run the default action
-	void ExecuteDefaultAction();
+	UE_API void ExecuteDefaultAction();
 
 	// Handles notifying the editor if the recent command mucked with the selection set
-	void PostExecutionCleanup();
+	UE_API void PostExecutionCleanup();
 };
+
+#undef UE_API
diff --git a/Engine/Source/Editor/Blutility/Classes/PlacedEditorUtilityBase.h b/Engine/Source/Editor/Blutility/Classes/PlacedEditorUtilityBase.h
--- a/Engine/Source/Editor/Blutility/Classes/PlacedEditorUtilityBase.h
+++ b/Engine/Source/Editor/Blutility/Classes/PlacedEditorUtilityBase.h
@@ -18,11 +18,13 @@
 
 #include "PlacedEditorUtilityBase.generated.h"
 
+#define UE_API BLUTILITY_API
+
 class UObject;
 struct FFrame;
 
-UCLASS(Abstract, hideCategories=(Object, Actor), Deprecated)
-class BLUTILITY_API ADEPRECATED_PlacedEditorUtilityBase : public AActor
+UCLASS(MinimalAPI, Abstract, hideCategories=(Object, Actor), Deprecated)
+class ADEPRECATED_PlacedEditorUtilityBase : public AActor
 {
 	GENERATED_UCLASS_BODY()
 
@@ -30,7 +32,7 @@ class BLUTILITY_API ADEPRECATED_PlacedEditorUtilityBase : public AActor
 	FString HelpText;
 
 	// AActor interface
-	virtual void TickActor(float DeltaSeconds, ELevelTick TickType, FActorTickFunction& ThisTickFunction) override;
+	UE_API virtual void TickActor(float DeltaSeconds, ELevelTick TickType, FActorTickFunction& ThisTickFunction) override;
 	virtual bool ShouldTickIfViewportsOnly() const override { return true; }
 	// End of AActor interface
 
@@ -38,7 +40,7 @@ class BLUTILITY_API ADEPRECATED_PlacedEditorUtilityBase : public AActor
 	 * Returns the current selection set in the editor.  Note that for non-editor builds, this will return an empty array
 	 */
 	UFUNCTION(BlueprintPure, Category = "Development|Editor")
-	TArray<AActor*> GetSelectionSet();
+	UE_API TArray<AActor*> GetSelectionSet();
 
 	/**
 	 * Gets information about the camera position for the primary level editor viewport.  In non-editor builds, these will be zeroed
@@ -48,7 +50,7 @@ class BLUTILITY_API ADEPRECATED_PlacedEditorUtilityBase : public AActor
 	 * @return	Whether or not we were able to get a camera for a level editing viewport
 	 */
 	UFUNCTION(BlueprintPure, Category = "Development|Editor")
-	bool GetLevelViewportCameraInfo(FVector& CameraLocation, FRotator& CameraRotation);
+	UE_API bool GetLevelViewportCameraInfo(FVector& CameraLocation, FRotator& CameraRotation);
 
 	/**
 	* Sets information about the camera position for the primary level editor viewport.
@@ -57,19 +59,19 @@ class BLUTILITY_API ADEPRECATED_PlacedEditorUtilityBase : public AActor
 	* @param	CameraRotation	Rotation the camera will be set to.
 	*/
 	UFUNCTION(BlueprintCallable, Category = "Development|Editor")
-	void SetLevelViewportCameraInfo(FVector CameraLocation, FRotator CameraRotation);
+	UE_API void SetLevelViewportCameraInfo(FVector CameraLocation, FRotator CameraRotation);
 
 	// Remove all actors from the selection set
 	UFUNCTION(BlueprintCallable, Category = "Development|Editor")
-	void ClearActorSelectionSet();
+	UE_API void ClearActorSelectionSet();
 
 	// Selects nothing in the editor (another way to clear the selection)
 	UFUNCTION(BlueprintCallable, Category = "Development|Editor")
-	void SelectNothing();
+	UE_API void SelectNothing();
 
 	// Set the selection state for the selected actor
 	UFUNCTION(BlueprintCallable, Category = "Development|Editor")
-	void SetActorSelectionState(AActor* Actor, bool bShouldBeSelected);
+	UE_API void SetActorSelectionState(AActor* Actor, bool bShouldBeSelected);
 
 	/**
 	* Attempts to find the actor specified by PathToActor in the current editor world
@@ -77,5 +79,7 @@ class BLUTILITY_API ADEPRECATED_PlacedEditorUtilityBase : public AActor
 	* @return	A reference to the actor, or none if it wasn't found
 	*/
 	UFUNCTION(BlueprintPure, Category = "Development|Editor")
-	AActor* GetActorReference(FString PathToActor);
+	UE_API AActor* GetActorReference(FString PathToActor);
 };
+
+#undef UE_API
diff --git a/Engine/Source/Editor/Blutility/Public/EditorUtilityLibrary.h b/Engine/Source/Editor/Blutility/Public/EditorUtilityLibrary.h
--- a/Engine/Source/Editor/Blutility/Public/EditorUtilityLibrary.h
+++ b/Engine/Source/Editor/Blutility/Public/EditorUtilityLibrary.h
@@ -18,6 +18,8 @@
 
 #include "EditorUtilityLibrary.generated.h"
 
+#define UE_API BLUTILITY_API
+
 class UWidget;
 class UEditorUtilityWidgetBlueprint;
 class AActor;
@@ -35,20 +37,20 @@ enum class ECastToWidgetBlueprintCases : uint8
 	CastFailed
 };
 
-UCLASS()
-class BLUTILITY_API UEditorUtilityBlueprintAsyncActionBase : public UBlueprintAsyncActionBase
+UCLASS(MinimalAPI)
+class UEditorUtilityBlueprintAsyncActionBase : public UBlueprintAsyncActionBase
 {
 	GENERATED_UCLASS_BODY()
 
 public:
-	virtual void RegisterWithGameInstance(const UObject* WorldContextObject) override;
-	virtual void SetReadyToDestroy() override;
+	UE_API virtual void RegisterWithGameInstance(const UObject* WorldContextObject) override;
+	UE_API virtual void SetReadyToDestroy() override;
 };
 
 DECLARE_DYNAMIC_MULTICAST_DELEGATE(FAsyncDelayComplete);
 
-UCLASS()
-class BLUTILITY_API UAsyncEditorDelay : public UEditorUtilityBlueprintAsyncActionBase
+UCLASS(MinimalAPI)
+class UAsyncEditorDelay : public UEditorUtilityBlueprintAsyncActionBase
 {
 	GENERATED_UCLASS_BODY()
 
@@ -56,7 +58,7 @@ public:
 #if WITH_EDITOR
 
 	UFUNCTION(BlueprintCallable, meta = (BlueprintInternalUseOnly = "true"))
-	static UAsyncEditorDelay* AsyncEditorDelay(float Seconds, int32 MinimumFrames = 30);
+	static UE_API UAsyncEditorDelay* AsyncEditorDelay(float Seconds, int32 MinimumFrames = 30);
 
 #endif
 
@@ -67,11 +69,11 @@ public:
 
 public:
 
-	void Start(float InMinimumSeconds, int32 InMinimumFrames);
+	UE_API void Start(float InMinimumSeconds, int32 InMinimumFrames);
 
 private:
 
-	bool HandleComplete(float DeltaTime);
+	UE_API bool HandleComplete(float DeltaTime);
 
 	uint64 EndFrame = 0;
 	double EndTime = 0;
@@ -79,8 +81,8 @@ private:
 
 DECLARE_DYNAMIC_MULTICAST_DELEGATE_OneParam(FAsyncEditorWaitForGameWorldEvent, UWorld*, World);
 
-UCLASS()
-class BLUTILITY_API UAsyncEditorWaitForGameWorld : public UEditorUtilityBlueprintAsyncActionBase
+UCLASS(MinimalAPI)
+class UAsyncEditorWaitForGameWorld : public UEditorUtilityBlueprintAsyncActionBase
 {
 	GENERATED_UCLASS_BODY()
 
@@ -88,7 +90,7 @@ public:
 #if WITH_EDITOR
 
 	UFUNCTION(BlueprintCallable, meta = (BlueprintInternalUseOnly = "true"))
-	static UAsyncEditorWaitForGameWorld* AsyncWaitForGameWorld(int32 Index = 0, bool Server = false);
+	static UE_API UAsyncEditorWaitForGameWorld* AsyncWaitForGameWorld(int32 Index = 0, bool Server = false);
 
 #endif
 
@@ -99,18 +101,18 @@ public:
 
 public:
 
-	void Start(int32 Index, bool Server);
+	UE_API void Start(int32 Index, bool Server);
 
 private:
 
-	bool OnTick(float DeltaTime);
+	UE_API bool OnTick(float DeltaTime);
 
 	int32 Index;
 	bool Server;
 };
 
-UCLASS()
-class BLUTILITY_API UAsyncEditorOpenMapAndFocusActor : public UEditorUtilityBlueprintAsyncActionBase
+UCLASS(MinimalAPI)
+class UAsyncEditorOpenMapAndFocusActor : public UEditorUtilityBlueprintAsyncActionBase
 {
 	GENERATED_UCLASS_BODY()
 
@@ -118,7 +120,7 @@ public:
 #if WITH_EDITOR
 
 	UFUNCTION(BlueprintCallable, meta = (BlueprintInternalUseOnly = "true"))
-	static UAsyncEditorOpenMapAndFocusActor* AsyncEditorOpenMapAndFocusActor(FSoftObjectPath Map, FString FocusActorName);
+	static UE_API UAsyncEditorOpenMapAndFocusActor* AsyncEditorOpenMapAndFocusActor(FSoftObjectPath Map, FString FocusActorName);
 
 #endif
 
@@ -129,11 +131,11 @@ public:
 
 public:
 
-	void Start(FSoftObjectPath InMap, FString InFocusActorName);
+	UE_API void Start(FSoftObjectPath InMap, FString InFocusActorName);
 
 private:
 
-	bool OnTick(float DeltaTime);
+	UE_API bool OnTick(float DeltaTime);
 
 	FSoftObjectPath Map;
 	FString FocusActorName;
@@ -141,8 +143,8 @@ private:
 
 
 // Expose editor utility functions to Blutilities 
-UCLASS()
-class BLUTILITY_API UEditorUtilityLibrary : public UBlueprintFunctionLibrary
+UCLASS(MinimalAPI)
+class UEditorUtilityLibrary : public UBlueprintFunctionLibrary
 {
 	GENERATED_UCLASS_BODY()
 
@@ -150,29 +152,29 @@ public:
 #if WITH_EDITOR
 
 	UFUNCTION(BlueprintCallable, Category = "Development|Editor")
-	static TArray<AActor*> GetSelectionSet();
+	static UE_API TArray<AActor*> GetSelectionSet();
 
 	UFUNCTION(BlueprintCallable, Category = "Development|Editor")
-	static void GetSelectionBounds(FVector& Origin, FVector& BoxExtent, float& SphereRadius);
+	static UE_API void GetSelectionBounds(FVector& Origin, FVector& BoxExtent, float& SphereRadius);
 
 	// Gets the set of currently selected assets
 	UFUNCTION(BlueprintCallable, Category = "Development|Editor")
-	static TArray<UObject*> GetSelectedAssets();
+	static UE_API TArray<UObject*> GetSelectedAssets();
 	
 	UFUNCTION(BlueprintCallable, Category = "Development|Editor")
-	static TArray<UObject*> GetSelectedAssetsOfClass(UClass* AssetClass);
+	static UE_API TArray<UObject*> GetSelectedAssetsOfClass(UClass* AssetClass);
 
 	// Gets the set of currently selected classes
 	UFUNCTION(BlueprintCallable, Category = "Development|Editor")
-	static TArray<UClass*> GetSelectedBlueprintClasses();
+	static UE_API TArray<UClass*> GetSelectedBlueprintClasses();
 
 	// Gets the set of currently selected asset data
 	UFUNCTION(BlueprintCallable, Category = "Development|Editor")
-	static TArray<FAssetData> GetSelectedAssetData();
+	static UE_API TArray<FAssetData> GetSelectedAssetData();
 
 	// Renames an asset (cannot move folders)
 	UFUNCTION(BlueprintCallable, Category = "Development|Editor")
-	static void RenameAsset(UObject* Asset, const FString& NewName);
+	static UE_API void RenameAsset(UObject* Asset, const FString& NewName);
 
 	/**
 	* Attempts to find the actor specified by PathToActor in the current editor world
@@ -180,7 +182,7 @@ public:
 	* @return	A reference to the actor, or none if it wasn't found
 	*/
 	UFUNCTION(BlueprintPure, Category = "Development|Editor")
-	AActor* GetActorReference(FString PathToActor);
+	UE_API AActor* GetActorReference(FString PathToActor);
 
 	/**
 	 * Attempts to get the path for the active content browser, returns false if there is no active content browser
@@ -189,36 +191,36 @@ public:
 	 * @return	Whether a path was successfully returned
 	 */
 	UFUNCTION(BlueprintPure, Category = "Development|Editor")
-	static bool GetCurrentContentBrowserPath(FString& OutPath);
+	static UE_API bool GetCurrentContentBrowserPath(FString& OutPath);
 
 	// Gets the current content browser path if one is open, whether it is internal or virtual.
 	UFUNCTION(BlueprintPure, Category = "Development|Editor")
-	static FContentBrowserItemPath GetCurrentContentBrowserItemPath();
+	static UE_API FContentBrowserItemPath GetCurrentContentBrowserItemPath();
 
 	// Gets the path to the currently selected folder in the content browser
 	UFUNCTION(BlueprintPure, Category = "Development|Editor")
-	static TArray<FString> GetSelectedFolderPaths();
+	static UE_API TArray<FString> GetSelectedFolderPaths();
 	
 	// Returns the folders that are selected in the path view for the content browser
     UFUNCTION(BlueprintPure, Category = "Development|Editor")
-    static TArray<FString> GetSelectedPathViewFolderPaths();
+    static UE_API TArray<FString> GetSelectedPathViewFolderPaths();
 
 	/**
 	 * Sync the Content Browser to the given folder(s)
 	 * @param	FolderList	The list of folders to sync to in the Content Browser
 	 */
 	UFUNCTION(BlueprintCallable, Category = "Editor Scripting | Content Browser")
-	static void SyncBrowserToFolders(const TArray<FString>& FolderList);
+	static UE_API void SyncBrowserToFolders(const TArray<FString>& FolderList);
 	
 	UFUNCTION(BlueprintCallable, Category = "Editor Scripting | Content Browser")
-	static void ConvertToEditorUtilityWidget(class UWidgetBlueprint* WidgetBP);
+	static UE_API void ConvertToEditorUtilityWidget(class UWidgetBlueprint* WidgetBP);
 	
 	UFUNCTION(BlueprintCallable, Category = "Editor Scripting | Widget Blueprint", Meta = (ExpandEnumAsExecs = "Branches"))
-	static void CastToWidgetBlueprint(UObject* Object, ECastToWidgetBlueprintCases& Branches, UWidgetBlueprint*& AsWidgetBlueprint);
+	static UE_API void CastToWidgetBlueprint(UObject* Object, ECastToWidgetBlueprintCases& Branches, UWidgetBlueprint*& AsWidgetBlueprint);
 
 	/** Searches the blueprint's widget hierarchy for a widget with the specified name */
 	UFUNCTION(BlueprintCallable, BlueprintPure = false, Category = "Editor Scripting | Widget Blueprint")
-	static UWidget* FindSourceWidgetByName(UWidgetBlueprint* WidgetBlueprint, FName WidgetName);
+	static UE_API UWidget* FindSourceWidgetByName(UWidgetBlueprint* WidgetBlueprint, FName WidgetName);
 
 	/**
 	 * Create a new widget and add to the specific widget blueprint's widget tree
@@ -231,6 +233,8 @@ public:
 	 * @return The widget that was created and added to the widget blueprint
 	 */
 	UFUNCTION(BlueprintCallable, BlueprintPure = false, Category = "Editor Scripting | Widget Blueprint")
-	static UWidget* AddSourceWidget(UWidgetBlueprint* WidgetBlueprint, TSubclassOf<UWidget> WidgetClass, FName WidgetName, FName WidgetParentName);
+	static UE_API UWidget* AddSourceWidget(UWidgetBlueprint* WidgetBlueprint, TSubclassOf<UWidget> WidgetClass, FName WidgetName, FName WidgetParentName);
 #endif
 };
+
+#undef UE_API
diff --git a/Engine/Source/Editor/Blutility/Public/EditorUtilitySubsystem.h b/Engine/Source/Editor/Blutility/Public/EditorUtilitySubsystem.h
--- a/Engine/Source/Editor/Blutility/Public/EditorUtilitySubsystem.h
+++ b/Engine/Source/Editor/Blutility/Public/EditorUtilitySubsystem.h
@@ -21,6 +21,8 @@
 
 #include "EditorUtilitySubsystem.generated.h"
 
+#define UE_API BLUTILITY_API
+
 class FOutputDevice;
 class FSpawnTabArgs;
 class FSubsystemCollectionBase;
@@ -40,21 +42,21 @@ struct FSoftObjectPath;
 /** Delegate for a PIE event exposed via Editor Utility (begin, end, pause/resume, etc) */
 DECLARE_DYNAMIC_MULTICAST_DELEGATE_OneParam(FOnEditorUtilityPIEEvent, const bool, bIsSimulating);
 
-UCLASS(config = EditorPerProjectUserSettings)
-class BLUTILITY_API UEditorUtilitySubsystem : public UEditorSubsystem
+UCLASS(MinimalAPI, config = EditorPerProjectUserSettings)
+class UEditorUtilitySubsystem : public UEditorSubsystem
 {
 	GENERATED_BODY()
 
 public:
-	UEditorUtilitySubsystem();
+	UE_API UEditorUtilitySubsystem();
 
-	virtual void Initialize(FSubsystemCollectionBase& Collection);
-	virtual void Deinitialize();
+	UE_API virtual void Initialize(FSubsystemCollectionBase& Collection);
+	UE_API virtual void Deinitialize();
 
-	static void AddReferencedObjects(UObject* InThis, FReferenceCollector& Collector);
+	static UE_API void AddReferencedObjects(UObject* InThis, FReferenceCollector& Collector);
 
-	void MainFrameCreationFinished(TSharedPtr<SWindow> InRootWindow, bool bIsRunningStartupDialog);
-	void HandleStartup();
+	UE_API void MainFrameCreationFinished(TSharedPtr<SWindow> InRootWindow, bool bIsRunningStartupDialog);
+	UE_API void HandleStartup();
 
 	UPROPERTY(config)
 	TArray<FSoftObjectPath> LoadedUIs;
@@ -67,75 +69,75 @@ public:
 
 	// Allow startup object to be garbage collected
 	UFUNCTION(BlueprintCallable, Category = "Development|Editor")
-	void ReleaseInstanceOfAsset(UObject* Asset);
+	UE_API void ReleaseInstanceOfAsset(UObject* Asset);
 
 	UFUNCTION(BlueprintCallable, Category = "Development|Editor")
-	bool TryRun(UObject* Asset);
+	UE_API bool TryRun(UObject* Asset);
 
 	UFUNCTION(BlueprintCallable, Category = "Development|Editor")
-	bool TryRunClass(UClass* ObjectClass);
+	UE_API bool TryRunClass(UClass* ObjectClass);
 
 	UFUNCTION(BlueprintCallable, Category = "Development|Editor")
-	bool CanRun(UObject* Asset) const;
+	UE_API bool CanRun(UObject* Asset) const;
 
 	UFUNCTION(BlueprintCallable, Category = "Development|Editor")
-	UEditorUtilityWidget* SpawnAndRegisterTabAndGetID(class UEditorUtilityWidgetBlueprint* InBlueprint, FName& NewTabID);
+	UE_API UEditorUtilityWidget* SpawnAndRegisterTabAndGetID(class UEditorUtilityWidgetBlueprint* InBlueprint, FName& NewTabID);
 
 	UFUNCTION(BlueprintCallable, Category = "Development|Editor")
-	UEditorUtilityWidget* SpawnAndRegisterTab(class UEditorUtilityWidgetBlueprint* InBlueprint);
+	UE_API UEditorUtilityWidget* SpawnAndRegisterTab(class UEditorUtilityWidgetBlueprint* InBlueprint);
 	
 	UFUNCTION(BlueprintCallable, Category = "Development|Editor")
-	void RegisterTabAndGetID(class UEditorUtilityWidgetBlueprint* InBlueprint, FName& NewTabID);
+	UE_API void RegisterTabAndGetID(class UEditorUtilityWidgetBlueprint* InBlueprint, FName& NewTabID);
 
 	/**
 	 * Unlike SpawnAndRegisterTabAndGetID allows spawn tab while providing TabID from Python scripts or BP
 	 */
 	UFUNCTION(BlueprintCallable, Category = "Development|Editor")
-	UEditorUtilityWidget* SpawnAndRegisterTabWithId(class UEditorUtilityWidgetBlueprint* InBlueprint, FName InTabID);
+	UE_API UEditorUtilityWidget* SpawnAndRegisterTabWithId(class UEditorUtilityWidgetBlueprint* InBlueprint, FName InTabID);
 
 	UFUNCTION(BlueprintCallable, Category = "Development|Editor")
-	UEditorUtilityWidget* SpawnAndRegisterTabAndGetIDGeneratedClass(UWidgetBlueprintGeneratedClass* InGeneratedWidgetBlueprint, FName& NewTabID);
+	UE_API UEditorUtilityWidget* SpawnAndRegisterTabAndGetIDGeneratedClass(UWidgetBlueprintGeneratedClass* InGeneratedWidgetBlueprint, FName& NewTabID);
 
 	UFUNCTION(BlueprintCallable, Category = "Development|Editor")
-	UEditorUtilityWidget* SpawnAndRegisterTabGeneratedClass(UWidgetBlueprintGeneratedClass* InGeneratedWidgetBlueprint);
+	UE_API UEditorUtilityWidget* SpawnAndRegisterTabGeneratedClass(UWidgetBlueprintGeneratedClass* InGeneratedWidgetBlueprint);
 
 	UFUNCTION(BlueprintCallable, Category = "Development|Editor")
-	void RegisterTabAndGetIDGeneratedClass(UWidgetBlueprintGeneratedClass* InGeneratedWidgetBlueprint, FName& NewTabID);
+	UE_API void RegisterTabAndGetIDGeneratedClass(UWidgetBlueprintGeneratedClass* InGeneratedWidgetBlueprint, FName& NewTabID);
 
 	/**
 	 * Unlike SpawnAndRegisterTabAndGetID allows spawn tab while providing TabID from Python scripts or BP
 	 */
 	UFUNCTION(BlueprintCallable, Category = "Development|Editor")
-	UEditorUtilityWidget* SpawnAndRegisterTabWithIdGeneratedClass(UWidgetBlueprintGeneratedClass* InGeneratedWidgetBlueprint, FName InTabID);
+	UE_API UEditorUtilityWidget* SpawnAndRegisterTabWithIdGeneratedClass(UWidgetBlueprintGeneratedClass* InGeneratedWidgetBlueprint, FName InTabID);
 
 	/** Given an ID for a tab, try to find a tab spawner that matches, and then spawn a tab. Returns true if it was able to find a matching tab spawner */
 	UFUNCTION(BlueprintCallable, Category = "Development|Editor")
-	bool SpawnRegisteredTabByID(FName NewTabID);
+	UE_API bool SpawnRegisteredTabByID(FName NewTabID);
 
 	/** Given an ID for a tab, try to find an existing tab. Returns true if it found a tab. */
 	UFUNCTION(BlueprintCallable, Category = "Development|Editor")
-	bool DoesTabExist(FName NewTabID);
+	UE_API bool DoesTabExist(FName NewTabID);
 
 	/** Given an ID for a tab, try to find and close an existing tab. Returns true if it found a tab to close. */
 	UFUNCTION(BlueprintCallable, Category = "Development|Editor")
-	bool CloseTabByID(FName NewTabID);
+	UE_API bool CloseTabByID(FName NewTabID);
 
 	/** Given an ID for a tab, try to close and unregister a tab that was registered through this subsystem */
 	UFUNCTION(BlueprintCallable, Category = "Development|Editor")
-	bool UnregisterTabByID(FName TabID);
+	UE_API bool UnregisterTabByID(FName TabID);
 
 	/** Given an editor utility widget blueprint, get the widget it creates. This will return a null pointer if the widget is not currently in a tab.*/
 	UFUNCTION(BlueprintCallable, Category = "Development|Editor")
-	UEditorUtilityWidget* FindUtilityWidgetFromBlueprint(class UEditorUtilityWidgetBlueprint* InBlueprint);
+	UE_API UEditorUtilityWidget* FindUtilityWidgetFromBlueprint(class UEditorUtilityWidgetBlueprint* InBlueprint);
 
 	/**  */
 	UFUNCTION(BlueprintCallable, Category = "Development|Editor")
-	void RegisterAndExecuteTask(UEditorUtilityTask* NewTask, UEditorUtilityTask* OptionalParentTask = nullptr);
+	UE_API void RegisterAndExecuteTask(UEditorUtilityTask* NewTask, UEditorUtilityTask* OptionalParentTask = nullptr);
 
-	void RemoveTaskFromActiveList(UEditorUtilityTask* Task);
+	UE_API void RemoveTaskFromActiveList(UEditorUtilityTask* Task);
 
-	void RegisterReferencedObject(UObject* ObjectToReference);
-	void UnregisterReferencedObject(UObject* ObjectToReference);
+	UE_API void RegisterReferencedObject(UObject* ObjectToReference);
+	UE_API void UnregisterReferencedObject(UObject* ObjectToReference);
 
 	/** Expose Begin PIE to blueprints.*/
 	UPROPERTY(BlueprintAssignable)
@@ -146,36 +148,36 @@ public:
 	FOnEditorUtilityPIEEvent OnEndPIE;
 
 protected:
-	void HandleStartupAssets();
-	void HandleDynamicStartupAssets(TConstArrayView<FAssetData> InAssets);
+	UE_API void HandleStartupAssets();
+	UE_API void HandleDynamicStartupAssets(TConstArrayView<FAssetData> InAssets);
 
-	void OnBlueprintGeneratedClassUnloaded(UBlueprintGeneratedClass* BPGC);
+	UE_API void OnBlueprintGeneratedClassUnloaded(UBlueprintGeneratedClass* BPGC);
 
 	UEditorUtilityTask* GetActiveTask() { return ActiveTaskStack.Num() > 0 ? ActiveTaskStack[ActiveTaskStack.Num() - 1] : nullptr; };
 
-	void StartTask(UEditorUtilityTask* Task);
+	UE_API void StartTask(UEditorUtilityTask* Task);
 
-	bool Tick(float DeltaTime);
+	UE_API bool Tick(float DeltaTime);
 
-	void ProcessRunTaskCommands();
+	UE_API void ProcessRunTaskCommands();
 
-	void RunTaskCommand(const TArray<FString>& Params, UWorld* InWorld, FOutputDevice& Ar);
-	void CancelAllTasksCommand(const TArray<FString>& Params, UWorld* InWorld, FOutputDevice& Ar);
+	UE_API void RunTaskCommand(const TArray<FString>& Params, UWorld* InWorld, FOutputDevice& Ar);
+	UE_API void CancelAllTasksCommand(const TArray<FString>& Params, UWorld* InWorld, FOutputDevice& Ar);
 
-	UClass* FindClassByName(const FString& RawTargetName);
-	UClass* FindBlueprintClass(const FString& TargetNameRaw);
+	UE_API UClass* FindClassByName(const FString& RawTargetName);
+	UE_API UClass* FindBlueprintClass(const FString& TargetNameRaw);
 
 	/** Called when Play in Editor begins. */
-	void HandleOnBeginPIE(const bool bIsSimulating);
+	UE_API void HandleOnBeginPIE(const bool bIsSimulating);
 
 	/** Called when Play in Editor stops. */
-	void HandleOnEndPIE(const bool bIsSimulating);
+	UE_API void HandleOnEndPIE(const bool bIsSimulating);
 
-	TSharedRef<SDockTab> SpawnEditorUITabFromGeneratedClass(const FSpawnTabArgs& SpawnTabArgs, UWidgetBlueprintGeneratedClass* InGeneratedWidgetBlueprint);
+	UE_API TSharedRef<SDockTab> SpawnEditorUITabFromGeneratedClass(const FSpawnTabArgs& SpawnTabArgs, UWidgetBlueprintGeneratedClass* InGeneratedWidgetBlueprint);
 
-	void OnSpawnedFromGeneratedClassTabClosed(TSharedRef<SDockTab> TabBeingClosed);
+	UE_API void OnSpawnedFromGeneratedClassTabClosed(TSharedRef<SDockTab> TabBeingClosed);
 
-	void OnMapChanged(UWorld* World, EMapChangeType MapChangeType);
+	UE_API void OnMapChanged(UWorld* World, EMapChangeType MapChangeType);
 
 	TMap<TSharedRef<SDockTab>, UEditorUtilityWidget*> SpawnedFromGeneratedClassTabs;
 
@@ -201,3 +203,5 @@ private:
 	UPROPERTY()
 	TSet<TObjectPtr<UObject>> ReferencedObjects;
 };
+
+#undef UE_API
diff --git a/Engine/Source/Editor/Blutility/Public/EditorUtilityWidgetProjectSettings.h b/Engine/Source/Editor/Blutility/Public/EditorUtilityWidgetProjectSettings.h
--- a/Engine/Source/Editor/Blutility/Public/EditorUtilityWidgetProjectSettings.h
+++ b/Engine/Source/Editor/Blutility/Public/EditorUtilityWidgetProjectSettings.h
@@ -12,6 +12,8 @@
 #include "WidgetEditingProjectSettings.h"
 #include "EditorUtilityWidgetProjectSettings.generated.h"
 
+#define UE_API BLUTILITY_API
+
 class UWidgetCompilerRule;
 class UUserWidget;
 class UWidgetBlueprint;
@@ -20,27 +22,27 @@ class UPanelWidget;
 /**
  * Implements the settings for Editor Utility Widget Editing Project Settings
  */
-UCLASS(config=Editor, defaultconfig)
-class BLUTILITY_API UEditorUtilityWidgetProjectSettings : public UWidgetEditingProjectSettings
+UCLASS(MinimalAPI, config=Editor, defaultconfig)
+class UEditorUtilityWidgetProjectSettings : public UWidgetEditingProjectSettings
 {
 	GENERATED_BODY()
 
 public:
 
-	virtual FText GetSectionText() const override;
-	virtual FText GetSectionDescription() const override;
+	UE_API virtual FText GetSectionText() const override;
+	UE_API virtual FText GetSectionDescription() const override;
 
 	// Begin UObject Interface
-	virtual void PostEditChangeChainProperty(struct FPropertyChangedChainEvent& PropertyChangedEvent) override;
+	UE_API virtual void PostEditChangeChainProperty(struct FPropertyChangedChainEvent& PropertyChangedEvent) override;
 	// End UObject Interface
 
 	/** Get the permission list that controls actor actions */
-	FNamePermissionList& GetAllowedEditorUtilityActorActions();
-	const FNamePermissionList& GetAllowedEditorUtilityActorActions() const;
+	UE_API FNamePermissionList& GetAllowedEditorUtilityActorActions();
+	UE_API const FNamePermissionList& GetAllowedEditorUtilityActorActions() const;
 
 	/** Get the permission list that controls asset actions */
-	FNamePermissionList& GetAllowedEditorUtilityAssetActions();
-	const FNamePermissionList& GetAllowedEditorUtilityAssetActions() const;
+	UE_API FNamePermissionList& GetAllowedEditorUtilityAssetActions();
+	UE_API const FNamePermissionList& GetAllowedEditorUtilityAssetActions() const;
 
 	/** Set true to search generated classes for scripted actions. */
 	UPROPERTY(EditAnywhere, config, Category = Designer)
@@ -54,3 +56,5 @@ protected:
 	/** EditorUtility asset actions to always consider during context menu creation */
 	FNamePermissionList AllowedEditorUtilityAssetActions;
 };
+
+#undef UE_API
diff --git a/Engine/Source/Editor/Blutility/Public/IEditorUtilityExtension.h b/Engine/Source/Editor/Blutility/Public/IEditorUtilityExtension.h
--- a/Engine/Source/Editor/Blutility/Public/IEditorUtilityExtension.h
+++ b/Engine/Source/Editor/Blutility/Public/IEditorUtilityExtension.h
@@ -8,13 +8,13 @@
 
 #include "IEditorUtilityExtension.generated.h"
 
-UINTERFACE(BlueprintType)
-class BLUTILITY_API UEditorUtilityExtension : public UInterface
+UINTERFACE(MinimalAPI, BlueprintType)
+class UEditorUtilityExtension : public UInterface
 {
 	GENERATED_UINTERFACE_BODY()
 };
 
-class BLUTILITY_API IEditorUtilityExtension
+class IEditorUtilityExtension
 {
 	GENERATED_IINTERFACE_BODY()
 };
diff --git a/Engine/Source/Editor/ClassViewer/Public/ClassViewerFilter.h b/Engine/Source/Editor/ClassViewer/Public/ClassViewerFilter.h
--- a/Engine/Source/Editor/ClassViewer/Public/ClassViewerFilter.h
+++ b/Engine/Source/Editor/ClassViewer/Public/ClassViewerFilter.h
@@ -15,6 +15,8 @@
 #include "UObject/SoftObjectPath.h"
 #include "UObject/TopLevelAssetPath.h"
 
+#define UE_API CLASSVIEWER_API
+
 class FClassViewerFilterFuncs;
 class FClassViewerInitializationOptions;
 class FClassViewerNode;
@@ -148,7 +150,7 @@ namespace EFilterReturn
 {
 	enum Type{ Failed = 0, Passed, NoItems };
 }
-class CLASSVIEWER_API FClassViewerFilterFuncs
+class FClassViewerFilterFuncs
 {
 public:
 
@@ -162,7 +164,7 @@ public:
 	 *
 	 * @return					EFilterReturn::Passed if it is a child-of a class in the set, EFilterReturn::Failed if it is not, EFilterReturn::NoItems if the set is empty.
 	 */
-	virtual EFilterReturn::Type IfInChildOfClassesSet(TSet< const UClass* >& InSet, const UClass* InClass);
+	UE_API virtual EFilterReturn::Type IfInChildOfClassesSet(TSet< const UClass* >& InSet, const UClass* InClass);
 
 	/** 
 	 * Checks if the given Class is a child-of any of the classes in a set.
@@ -172,7 +174,7 @@ public:
 	 *
 	 * @return					EFilterReturn::Passed if it is a child-of a class in the set, EFilterReturn::Failed if it is not, EFilterReturn::NoItems if the set is empty.
 	 */
-	virtual EFilterReturn::Type IfInChildOfClassesSet(TSet< const UClass* >& InSet, const TSharedPtr< const class IUnloadedBlueprintData > InClass);
+	UE_API virtual EFilterReturn::Type IfInChildOfClassesSet(TSet< const UClass* >& InSet, const TSharedPtr< const class IUnloadedBlueprintData > InClass);
 
 	/** 
 	 * Checks if the given Class is a child-of ALL of the classes in a set.
@@ -182,7 +184,7 @@ public:
 	 *
 	 * @return					EFilterReturn::Passed if it is a child-of ALL the classes in the set, EFilterReturn::Failed if it is not, EFilterReturn::NoItems if the set is empty.
 	 */
-	virtual EFilterReturn::Type IfMatchesAllInChildOfClassesSet(TSet< const UClass* >& InSet, const UClass* InClass);
+	UE_API virtual EFilterReturn::Type IfMatchesAllInChildOfClassesSet(TSet< const UClass* >& InSet, const UClass* InClass);
 
 	/** 
 	 * Checks if the given Class is a child-of ALL of the classes in a set.
@@ -192,7 +194,7 @@ public:
 	 *
 	 * @return					EFilterReturn::Passed if it is a child-of ALL the classes in the set, EFilterReturn::Failed if it is not, EFilterReturn::NoItems if the set is empty.
 	 */
-	virtual EFilterReturn::Type IfMatchesAllInChildOfClassesSet(TSet< const UClass* >& InSet, const TSharedPtr< const class IUnloadedBlueprintData > InClass);
+	UE_API virtual EFilterReturn::Type IfMatchesAllInChildOfClassesSet(TSet< const UClass* >& InSet, const TSharedPtr< const class IUnloadedBlueprintData > InClass);
 
 	/** 
 	 * Checks if ALL the Objects has a Is-A relationship with the passed in class.
@@ -202,7 +204,7 @@ public:
 	 *
 	 * @return					EFilterReturn::Passed if ALL the Objects set IsA the passed class, EFilterReturn::Failed if it is not, EFilterReturn::NoItems if the set is empty.
 	 */
-	virtual EFilterReturn::Type IfMatchesAll_ObjectsSetIsAClass(TSet< const UObject* >& InSet, const UClass* InClass);
+	UE_API virtual EFilterReturn::Type IfMatchesAll_ObjectsSetIsAClass(TSet< const UObject* >& InSet, const UClass* InClass);
 
 	/** 
 	 * Checks if ALL the Objects has a Is-A relationship with the passed in class.
@@ -212,7 +214,7 @@ public:
 	 *
 	 * @return					EFilterReturn::Passed if ALL the Objects set IsA the passed class, EFilterReturn::Failed if it is not, EFilterReturn::NoItems if the set is empty.
 	 */
-	virtual EFilterReturn::Type IfMatchesAll_ObjectsSetIsAClass(TSet< const UObject* >& InSet, const TSharedPtr< const class IUnloadedBlueprintData > InClass);
+	UE_API virtual EFilterReturn::Type IfMatchesAll_ObjectsSetIsAClass(TSet< const UObject* >& InSet, const TSharedPtr< const class IUnloadedBlueprintData > InClass);
 
 	/** 
 	 * Checks if ALL the Classes set has a Is-A relationship with the passed in class.
@@ -222,7 +224,7 @@ public:
 	 *
 	 * @return					EFilterReturn::Passed if ALL the Classes set IsA the passed class, EFilterReturn::Failed if it is not, EFilterReturn::NoItems if the set is empty.
 	 */
-	virtual EFilterReturn::Type IfMatchesAll_ClassesSetIsAClass(TSet< const UClass* >& InSet, const UClass* InClass);
+	UE_API virtual EFilterReturn::Type IfMatchesAll_ClassesSetIsAClass(TSet< const UClass* >& InSet, const UClass* InClass);
 
 	/** 
 	 * Checks if ALL the Classes set has a Is-A relationship with the passed in class.
@@ -232,7 +234,7 @@ public:
 	 *
 	 * @return					EFilterReturn::Passed if ALL the Classes set IsA the passed class, EFilterReturn::Failed if it is not, EFilterReturn::NoItems if the set is empty.
 	 */
-	virtual EFilterReturn::Type IfMatchesAll_ClassesSetIsAClass(TSet< const UClass* >& InSet, const TSharedPtr< const class IUnloadedBlueprintData > InClass);
+	UE_API virtual EFilterReturn::Type IfMatchesAll_ClassesSetIsAClass(TSet< const UClass* >& InSet, const TSharedPtr< const class IUnloadedBlueprintData > InClass);
 
 	/** 
 	 * Checks if any in the Classes set has a Is-A relationship with the passed in class.
@@ -242,7 +244,7 @@ public:
 	 *
 	 * @return					EFilterReturn::Passed if the Classes set IsA the passed class, EFilterReturn::Failed if it is not, EFilterReturn::NoItems if the set is empty.
 	 */
-	virtual EFilterReturn::Type IfMatches_ClassesSetIsAClass(TSet< const UClass* >& InSet, const UClass* InClass);
+	UE_API virtual EFilterReturn::Type IfMatches_ClassesSetIsAClass(TSet< const UClass* >& InSet, const UClass* InClass);
 
 	/** 
 	 * Checks if any in the Classes set has a Is-A relationship with the passed in class.
@@ -252,7 +254,7 @@ public:
 	 *
 	 * @return					EFilterReturn::Passed if the Classes set IsA the passed class, EFilterReturn::Failed if it is not, EFilterReturn::NoItems if the set is empty.
 	 */
-	virtual EFilterReturn::Type IfMatches_ClassesSetIsAClass(TSet< const UClass* >& InSet, const TSharedPtr< const class IUnloadedBlueprintData > InClass);
+	UE_API virtual EFilterReturn::Type IfMatches_ClassesSetIsAClass(TSet< const UClass* >& InSet, const TSharedPtr< const class IUnloadedBlueprintData > InClass);
 
 	/** 
 	 * Checks if the Class is in the Classes set.
@@ -262,7 +264,7 @@ public:
 	 *
 	 * @return					EFilterReturn::Passed if the Class is in the Classes set, EFilterReturn::Failed if it is not, EFilterReturn::NoItems if the set is empty.
 	 */
-	virtual EFilterReturn::Type IfInClassesSet(TSet< const UClass* >& InSet, const UClass* InClass);
+	UE_API virtual EFilterReturn::Type IfInClassesSet(TSet< const UClass* >& InSet, const UClass* InClass);
 
 	/** 
 	 * Checks if the Class is in the Classes set.
@@ -272,7 +274,7 @@ public:
 	 *
 	 * @return					EFilterReturn::Passed if the Class is in the Classes set, EFilterReturn::Failed if it is not, EFilterReturn::NoItems if the set is empty.
 	 */
-	virtual EFilterReturn::Type IfInClassesSet(TSet< const UClass* >& InSet, const TSharedPtr< const class IUnloadedBlueprintData > InClass);
+	UE_API virtual EFilterReturn::Type IfInClassesSet(TSet< const UClass* >& InSet, const TSharedPtr< const class IUnloadedBlueprintData > InClass);
 };
 
 class IUnloadedBlueprintData
@@ -369,3 +371,5 @@ public:
 	 */
 	virtual FTopLevelAssetPath GetClassPathName() const = 0;
 };
+
+#undef UE_API
diff --git a/Engine/Source/Editor/ClassViewer/Public/ClassViewerProjectSettings.h b/Engine/Source/Editor/ClassViewer/Public/ClassViewerProjectSettings.h
--- a/Engine/Source/Editor/ClassViewer/Public/ClassViewerProjectSettings.h
+++ b/Engine/Source/Editor/ClassViewer/Public/ClassViewerProjectSettings.h
@@ -20,8 +20,8 @@ struct FSoftClassPath;
 /**
  * Implements the settings for the Class Viewer Project Settings
  */
-UCLASS(config=Engine, defaultconfig)
-class CLASSVIEWER_API UClassViewerProjectSettings : public UObject
+UCLASS(MinimalAPI, config=Engine, defaultconfig)
+class UClassViewerProjectSettings : public UObject
 {
 	GENERATED_UCLASS_BODY()
 
diff --git a/Engine/Source/Editor/ClothPainter/Private/SClothPaintTab.h b/Engine/Source/Editor/ClothPainter/Private/SClothPaintTab.h
--- a/Engine/Source/Editor/ClothPainter/Private/SClothPaintTab.h
+++ b/Engine/Source/Editor/ClothPainter/Private/SClothPaintTab.h
@@ -8,6 +8,8 @@
 #include "Widgets/DeclarativeSyntaxSupport.h"
 #include "Widgets/SCompoundWidget.h"
 
+#define UE_API CLOTHPAINTER_API
+
 class IDetailsView;
 class IPersonaToolkit;
 class ISkeletalMeshEditor;
@@ -17,7 +19,7 @@ class SScrollBox;
 class UClothingAssetCommon;
 struct FGeometry;
 
-class CLOTHPAINTER_API SClothPaintTab : public SCompoundWidget
+class SClothPaintTab : public SCompoundWidget
 {
 public:
 	SLATE_BEGIN_ARGS(SClothPaintTab)
@@ -25,28 +27,28 @@ public:
 	SLATE_ARGUMENT(TWeakPtr<class FAssetEditorToolkit>, InHostingApp)
 	SLATE_END_ARGS()
 
-	SClothPaintTab();
-	~SClothPaintTab();
+	UE_API SClothPaintTab();
+	UE_API ~SClothPaintTab();
 
 	/** SWidget functions */
-	void Construct(const FArguments& InArgs);
-	virtual void Tick(const FGeometry& AllottedGeometry, const double InCurrentTime, const float InDeltaTime) override;
+	UE_API void Construct(const FArguments& InArgs);
+	UE_API virtual void Tick(const FGeometry& AllottedGeometry, const double InCurrentTime, const float InDeltaTime) override;
 
 	/** Setup and teardown the cloth paint UI */
-	void EnterPaintMode();
-	void ExitPaintMode();
+	UE_API void EnterPaintMode();
+	UE_API void ExitPaintMode();
 
 protected:
 
 	/** Called from the selector when the asset selection changes (Asset, LOD, Mask) */
-	void OnAssetSelectionChanged(TWeakObjectPtr<UClothingAssetCommon> InAssetPtr, int32 InLodIndex, int32 InMaskIndex);
+	UE_API void OnAssetSelectionChanged(TWeakObjectPtr<UClothingAssetCommon> InAssetPtr, int32 InLodIndex, int32 InMaskIndex);
 
 	/** Whether or not the asset config section is enabled for editing */
-	bool IsAssetDetailsPanelEnabled();
+	UE_API bool IsAssetDetailsPanelEnabled();
 
 	/** Helpers for getting editor objects */
-	ISkeletalMeshEditor* GetSkeletalMeshEditor() const;
-	TSharedRef<IPersonaToolkit> GetPersonaToolkit() const;
+	UE_API ISkeletalMeshEditor* GetSkeletalMeshEditor() const;
+	UE_API TSharedRef<IPersonaToolkit> GetPersonaToolkit() const;
 
 	TWeakPtr<class FAssetEditorToolkit> HostingApp;
 	
@@ -57,3 +59,5 @@ protected:
 
 	bool bModeApplied;
 };
+
+#undef UE_API
diff --git a/Engine/Source/Editor/ClothPainter/Public/ClothingAssetExporter.h b/Engine/Source/Editor/ClothPainter/Public/ClothingAssetExporter.h
--- a/Engine/Source/Editor/ClothPainter/Public/ClothingAssetExporter.h
+++ b/Engine/Source/Editor/ClothPainter/Public/ClothingAssetExporter.h
@@ -19,7 +19,7 @@ class UClothingAssetExporter;
 template <typename FuncType> class TFunctionRef;
 
 /** A modular interface to provide ways to export clothing assets. */
-class CLOTHPAINTER_API IClothingAssetExporterClassProvider : public IModularFeature
+class IClothingAssetExporterClassProvider : public IModularFeature
 {
 public:
 	inline static const FName FeatureName = TEXT("ClothingAssetExporterClassProvider");
@@ -28,8 +28,8 @@ public:
 };
 
 /** Modular exporter base class. */
-UCLASS(Abstract)
-class CLOTHPAINTER_API UClothingAssetExporter : public UObject
+UCLASS(MinimalAPI, Abstract)
+class UClothingAssetExporter : public UObject
 {
 	GENERATED_BODY()
 
diff --git a/Engine/Source/Editor/ClothingSystemEditor/Public/ClothingAssetFactory.h b/Engine/Source/Editor/ClothingSystemEditor/Public/ClothingAssetFactory.h
--- a/Engine/Source/Editor/ClothingSystemEditor/Public/ClothingAssetFactory.h
+++ b/Engine/Source/Editor/ClothingSystemEditor/Public/ClothingAssetFactory.h
@@ -15,6 +15,8 @@
 
 #include "ClothingAssetFactory.generated.h"
 
+#define UE_API CLOTHINGSYSTEMEDITOR_API
+
 class UClothingAssetBase;
 class UObject;
 class USkeletalMesh;
@@ -40,34 +42,36 @@ namespace NvParameterized
 	class Interface;
 }
 
-UCLASS(hidecategories=Object)
-class CLOTHINGSYSTEMEDITOR_API UClothingAssetFactory : public UClothingAssetFactoryBase
+UCLASS(MinimalAPI, hidecategories=Object)
+class UClothingAssetFactory : public UClothingAssetFactoryBase
 {
 	GENERATED_BODY()
 
 public:
 
-	UClothingAssetFactory(const FObjectInitializer& ObjectInitializer);
+	UE_API UClothingAssetFactory(const FObjectInitializer& ObjectInitializer);
 
 	// Import the given file, treating it as an APEX asset file and return the resulting asset
-	virtual UClothingAssetBase* Import(const FString& Filename, USkeletalMesh* TargetMesh, FName InName = NAME_None) override;
-	virtual UClothingAssetBase* Reimport(const FString& Filename, USkeletalMesh* TargetMesh, UClothingAssetBase* OriginalAsset) override;
-	virtual UClothingAssetBase* CreateFromSkeletalMesh(USkeletalMesh* TargetMesh, FSkeletalMeshClothBuildParams& Params) override;
-	virtual UClothingAssetBase* CreateFromExistingCloth(USkeletalMesh* TargetMesh, USkeletalMesh* SourceMesh, UClothingAssetBase* SourceAsset) override;
-	virtual UClothingAssetBase* ImportLodToClothing(USkeletalMesh* TargetMesh, FSkeletalMeshClothBuildParams& Params) override;
+	UE_API virtual UClothingAssetBase* Import(const FString& Filename, USkeletalMesh* TargetMesh, FName InName = NAME_None) override;
+	UE_API virtual UClothingAssetBase* Reimport(const FString& Filename, USkeletalMesh* TargetMesh, UClothingAssetBase* OriginalAsset) override;
+	UE_API virtual UClothingAssetBase* CreateFromSkeletalMesh(USkeletalMesh* TargetMesh, FSkeletalMeshClothBuildParams& Params) override;
+	UE_API virtual UClothingAssetBase* CreateFromExistingCloth(USkeletalMesh* TargetMesh, USkeletalMesh* SourceMesh, UClothingAssetBase* SourceAsset) override;
+	UE_API virtual UClothingAssetBase* ImportLodToClothing(USkeletalMesh* TargetMesh, FSkeletalMeshClothBuildParams& Params) override;
 	// Tests whether the given filename should be able to be imported
-	virtual bool CanImport(const FString& Filename) override;
+	UE_API virtual bool CanImport(const FString& Filename) override;
 
 	// Given an APEX asset, build a UClothingAssetCommon containing the required data
-	virtual UClothingAssetBase* CreateFromApexAsset(nvidia::apex::ClothingAsset* InApexAsset, USkeletalMesh* TargetMesh, FName InName = NAME_None) override;
+	UE_API virtual UClothingAssetBase* CreateFromApexAsset(nvidia::apex::ClothingAsset* InApexAsset, USkeletalMesh* TargetMesh, FName InName = NAME_None) override;
 
 private:
 
 	// Utility methods for skeletal mesh extraction //////////////////////////
 
 	/** Handles internal import of LODs */
-	bool ImportToLodInternal(USkeletalMesh* SourceMesh, int32 SourceLodIndex, int32 SourceSectionIndex, UClothingAssetCommon* DestAsset, FClothLODDataCommon& DestLod, int32 DestLodIndex, const FClothLODDataCommon* InParameterRemapSource = nullptr);
+	UE_API bool ImportToLodInternal(USkeletalMesh* SourceMesh, int32 SourceLodIndex, int32 SourceSectionIndex, UClothingAssetCommon* DestAsset, FClothLODDataCommon& DestLod, int32 DestLodIndex, const FClothLODDataCommon* InParameterRemapSource = nullptr);
 
 	//////////////////////////////////////////////////////////////////////////
 
 };
+
+#undef UE_API
diff --git a/Engine/Source/Editor/ClothingSystemEditorInterface/Public/ClothingAssetFactoryInterface.h b/Engine/Source/Editor/ClothingSystemEditorInterface/Public/ClothingAssetFactoryInterface.h
--- a/Engine/Source/Editor/ClothingSystemEditorInterface/Public/ClothingAssetFactoryInterface.h
+++ b/Engine/Source/Editor/ClothingSystemEditorInterface/Public/ClothingAssetFactoryInterface.h
@@ -14,6 +14,8 @@
 
 #include "ClothingAssetFactoryInterface.generated.h"
 
+#define UE_API CLOTHINGSYSTEMEDITORINTERFACE_API
+
 namespace nvidia
 {
 	namespace apex
@@ -28,8 +30,8 @@ struct FSkeletalMeshClothBuildParams;
 
 // Clothing asset factories should inherit this interface/uobject to provide functionality
 // to build clothing assets from .apx files imported to the engine
-UCLASS(Abstract)
-class CLOTHINGSYSTEMEDITORINTERFACE_API UClothingAssetFactoryBase : public UObject
+UCLASS(MinimalAPI, Abstract)
+class UClothingAssetFactoryBase : public UObject
 {
 	GENERATED_BODY()
 
@@ -98,12 +100,14 @@ public:
 
 // An interface for a class that will provide a clothing asset factory, this should be
 // registered as a feature under its FeatureName to be picked up by the engine
-class CLOTHINGSYSTEMEDITORINTERFACE_API IClothingAssetFactoryProvider : public IModularFeature
+class IClothingAssetFactoryProvider : public IModularFeature
 {
 public:
-	static const FName FeatureName;
+	static UE_API const FName FeatureName;
 
 	// Called by the engine to retrieve a valid factory from a provider
 	// This can be the default object for the factory class or a full instance
 	virtual UClothingAssetFactoryBase* GetFactory() = 0;
 };
+
+#undef UE_API
diff --git a/Engine/Source/Editor/ClothingSystemEditorInterface/Public/ClothingSystemEditorInterfaceModule.h b/Engine/Source/Editor/ClothingSystemEditorInterface/Public/ClothingSystemEditorInterfaceModule.h
--- a/Engine/Source/Editor/ClothingSystemEditorInterface/Public/ClothingSystemEditorInterfaceModule.h
+++ b/Engine/Source/Editor/ClothingSystemEditorInterface/Public/ClothingSystemEditorInterfaceModule.h
@@ -5,24 +5,28 @@
 #include "Modules/ModuleInterface.h"
 #include "UObject/NameTypes.h"
 
+#define UE_API CLOTHINGSYSTEMEDITORINTERFACE_API
+
 class ISimulationEditorExtender;
 class UClothingAssetFactoryBase;
 
-class CLOTHINGSYSTEMEDITORINTERFACE_API FClothingSystemEditorInterfaceModule : public IModuleInterface
+class FClothingSystemEditorInterfaceModule : public IModuleInterface
 {
 
 public:
 
-	const static FName ExtenderFeatureName;
+	UE_API const static FName ExtenderFeatureName;
 
-	FClothingSystemEditorInterfaceModule();
+	UE_API FClothingSystemEditorInterfaceModule();
 
-	virtual void StartupModule() override;
-	virtual void ShutdownModule() override;
+	UE_API virtual void StartupModule() override;
+	UE_API virtual void ShutdownModule() override;
 
-	UClothingAssetFactoryBase* GetClothingAssetFactory();
-	ISimulationEditorExtender* GetSimulationEditorExtender(FName InSimulationFactoryClassName);
+	UE_API UClothingAssetFactoryBase* GetClothingAssetFactory();
+	UE_API ISimulationEditorExtender* GetSimulationEditorExtender(FName InSimulationFactoryClassName);
 
 private:
 
 };
+
+#undef UE_API
