Fixups after integration of unrealcodefixup changes.

diff --git a/Engine/Source/Editor/Blutility/Classes/EditorUtilityWidget.h b/Engine/Source/Editor/Blutility/Classes/EditorUtilityWidget.h
--- a/Engine/Source/Editor/Blutility/Classes/EditorUtilityWidget.h
+++ b/Engine/Source/Editor/Blutility/Classes/EditorUtilityWidget.h
@@ -67,7 +67,7 @@ public:
 protected:
 
 	/** Used to send event OnEditorToolStarted event on initialize */
-	virtual void NativeOnInitialized() override;
+	UE_API virtual void NativeOnInitialized() override;
 
 	/** The display name for tabs spawned with this widget */
 	UPROPERTY(Category = Config, EditDefaultsOnly, BlueprintReadWrite, AssetRegistrySearchable)
