Fixups after integration of unrealcodefixup changes.

diff --git a/Engine/Source/Editor/Blutility/Classes/EditorUtilityObject.h b/Engine/Source/Editor/Blutility/Classes/EditorUtilityObject.h
--- a/Engine/Source/Editor/Blutility/Classes/EditorUtilityObject.h
+++ b/Engine/Source/Editor/Blutility/Classes/EditorUtilityObject.h
@@ -26,7 +26,7 @@ class UEditorUtilityObject : public UObject, public IAssetRegistryTagProviderInt
 	UE_API void Run();
 
 	// we use this to call tool started delegate
-	virtual void PostLoad() override;
+	UE_API virtual void PostLoad() override;
 
 	//~ Begin IAssetRegistryTagProviderInterface interface
 	virtual bool ShouldAddCDOTagsToBlueprintClass() const override 
