MVVM: fix bindings library unloading after re-compiling EUW
The old editor utility userwidget might have the updated ViewClass, so the destruct should only be called if the views actually are derived from this viewclass
also prevent EUW from binding to OnCompile everytime it is started. Only one binding is enough.

diff --git a/Engine/Plugins/Runtime/ModelViewViewModel/Source/ModelViewViewModel/Private/View/MVVMView.cpp b/Engine/Plugins/Runtime/ModelViewViewModel/Source/ModelViewViewModel/Private/View/MVVMView.cpp
--- a/Engine/Plugins/Runtime/ModelViewViewModel/Source/ModelViewViewModel/Private/View/MVVMView.cpp
+++ b/Engine/Plugins/Runtime/ModelViewViewModel/Source/ModelViewViewModel/Private/View/MVVMView.cpp
@@ -105,7 +105,7 @@ void UMVVMView::Construct()
 #if WITH_EDITOR
 	if (GEditor)
 	{
-		if (GetWorld()->WorldType == EWorldType::Editor)
+		if (GetWorld() && GetWorld()->WorldType == EWorldType::Editor)
 		{
 			FModelViewViewModelModule& Module = FModuleManager::GetModuleChecked<FModelViewViewModelModule>("ModelViewViewModel");
 			BlueprintPreCompileHandle = Module.OnBlueprintPreCompile.AddUObject(this, &UMVVMView::HandleBlueprintPreCompile);
@@ -129,7 +129,7 @@ void UMVVMView::Destruct()
 		BlueprintPreCompileHandle.Reset();
 	}
 
-	if (!GeneratedViewClass && GetWorld()->WorldType == EWorldType::Editor)
+	if (!GeneratedViewClass && (!GetWorld() || GetWorld()->WorldType == EWorldType::Editor))
 	{
 		return;
 	}
diff --git a/Engine/Plugins/Runtime/ModelViewViewModel/Source/ModelViewViewModel/Private/View/MVVMViewClass.cpp b/Engine/Plugins/Runtime/ModelViewViewModel/Source/ModelViewViewModel/Private/View/MVVMViewClass.cpp
--- a/Engine/Plugins/Runtime/ModelViewViewModel/Source/ModelViewViewModel/Private/View/MVVMViewClass.cpp
+++ b/Engine/Plugins/Runtime/ModelViewViewModel/Source/ModelViewViewModel/Private/View/MVVMViewClass.cpp
@@ -752,7 +752,6 @@ void UMVVMViewClass::Construct(UUserWidget* UserWidget)
 		FModelViewViewModelModule& Module = FModuleManager::GetModuleChecked<FModelViewViewModelModule>("ModelViewViewModel");
 		Module.OnBlueprintCompiled.AddUObject(this, &UMVVMViewClass::HandleBlueprintCompiled);
 	}
-	bConstructed = true;
 #endif
 }
 
@@ -760,16 +759,34 @@ void UMVVMViewClass::Construct(UUserWidget* UserWidget)
 void UMVVMViewClass::Destruct(UUserWidget* UserWidget)
 {
 #if WITH_EDITOR
-	if (bConstructed)
-#endif
+	if (UserWidget)
 	{
-		--ViewCounter;
-		ensure(ViewCounter >= 0);
-		if (ViewCounter == 0)
+		bool bIsAnyViewFromThisClass = false;
+		for (UUserWidgetExtension* Extension : UserWidget->GetExtensions(UMVVMView::StaticClass()))
+		{
+			if (UMVVMView* View = Cast<UMVVMView>(Extension))
+			{
+				if (View->GetViewClass() == this)
+				{
+					bIsAnyViewFromThisClass = true;
+					break;
+				}
+			}
+		}
+
+		if (!bIsAnyViewFromThisClass)
 		{
-			BindingLibrary.Unload();
+			return;
 		}
 	}
+#endif
+
+	--ViewCounter;
+	ensure(ViewCounter >= 0);
+	if (ViewCounter == 0)
+	{
+		BindingLibrary.Unload();
+	}
 
 #if WITH_EDITOR
 	if (GEditor && BlueprintCompiledHandle.IsValid())
diff --git a/Engine/Plugins/Runtime/ModelViewViewModel/Source/ModelViewViewModel/Public/View/MVVMViewClass.h b/Engine/Plugins/Runtime/ModelViewViewModel/Source/ModelViewViewModel/Public/View/MVVMViewClass.h
--- a/Engine/Plugins/Runtime/ModelViewViewModel/Source/ModelViewViewModel/Public/View/MVVMViewClass.h
+++ b/Engine/Plugins/Runtime/ModelViewViewModel/Source/ModelViewViewModel/Public/View/MVVMViewClass.h
@@ -798,7 +798,6 @@ private:
 
 #if WITH_EDITORONLY_DATA
 	FDelegateHandle BlueprintCompiledHandle;
-	bool bConstructed = false;
 #endif
 };
 
diff --git a/Engine/Source/Editor/Blutility/Classes/EditorUtilityWidgetBlueprint.h b/Engine/Source/Editor/Blutility/Classes/EditorUtilityWidgetBlueprint.h
--- a/Engine/Source/Editor/Blutility/Classes/EditorUtilityWidgetBlueprint.h
+++ b/Engine/Source/Editor/Blutility/Classes/EditorUtilityWidgetBlueprint.h
@@ -112,6 +112,7 @@ private:
 	UPROPERTY(Transient)
 	TObjectPtr<UEditorUtilityWidget> CreatedUMGWidget;
 
+	FDelegateHandle OnBlueprintCompileHandle;
 };
 
 #undef UE_API
diff --git a/Engine/Source/Editor/Blutility/Private/EditorUtilityWidgetBlueprint.cpp b/Engine/Source/Editor/Blutility/Private/EditorUtilityWidgetBlueprint.cpp
--- a/Engine/Source/Editor/Blutility/Private/EditorUtilityWidgetBlueprint.cpp
+++ b/Engine/Source/Editor/Blutility/Private/EditorUtilityWidgetBlueprint.cpp
@@ -77,8 +77,9 @@ TSharedRef<SDockTab> UEditorUtilityWidgetBlueprint::SpawnEditorUITab(const FSpaw
 	
 		SpawnedTab->SetOnTabClosed(SDockTab::FOnTabClosedCallback::CreateUObject(this, &UEditorUtilityWidgetBlueprint::UpdateRespawnListIfNeeded));
 		CreatedTab = SpawnedTab;
-	
-		OnCompiled().AddUObject(this, &UEditorUtilityWidgetBlueprint::RegenerateCreatedTab);
+
+		OnCompiled().Remove(OnBlueprintCompileHandle);
+		OnBlueprintCompileHandle = OnCompiled().AddUObject(this, &UEditorUtilityWidgetBlueprint::RegenerateCreatedTab);
 	
 		FLevelEditorModule& LevelEditor = FModuleManager::LoadModuleChecked<FLevelEditorModule>("LevelEditor");
 		LevelEditor.OnMapChanged().AddUObject(this, &UEditorUtilityWidgetBlueprint::ChangeTabWorld);
@@ -94,7 +95,8 @@ TSharedRef<SDockTab> UEditorUtilityWidgetBlueprint::SpawnEditorUITab(const FSpaw
 	SpawnedTab->SetOnTabClosed(SDockTab::FOnTabClosedCallback::CreateUObject(this, &UEditorUtilityWidgetBlueprint::UpdateRespawnListIfNeeded));
 	CreatedTab = SpawnedTab;
 	
-	OnCompiled().AddUObject(this, &UEditorUtilityWidgetBlueprint::RegenerateCreatedTab);
+	OnCompiled().Remove(OnBlueprintCompileHandle);
+	OnBlueprintCompileHandle = OnCompiled().AddUObject(this, &UEditorUtilityWidgetBlueprint::RegenerateCreatedTab);
 	
 	FLevelEditorModule& LevelEditor = FModuleManager::LoadModuleChecked<FLevelEditorModule>("LevelEditor");
 	LevelEditor.OnMapChanged().AddUObject(this, &UEditorUtilityWidgetBlueprint::ChangeTabWorld);
