Fixed scripted asset actions not applying to cooked BPs

#rnx

diff --git a/Engine/Source/Editor/Blutility/Private/BlutilityContentBrowserExtensions.cpp b/Engine/Source/Editor/Blutility/Private/BlutilityContentBrowserExtensions.cpp
--- a/Engine/Source/Editor/Blutility/Private/BlutilityContentBrowserExtensions.cpp
+++ b/Engine/Source/Editor/Blutility/Private/BlutilityContentBrowserExtensions.cpp
@@ -106,12 +106,20 @@ void FBlutilityContentBrowserExtensions::RegisterMenus()
 						bool bPassesClassFilter = false;
 						if (bIsActionForBlueprints)
 						{
-							if (TSubclassOf<UBlueprint> AssetClass = SelectedAssetClasses[SelectedAssetIndex])
+							UClass* AssetClass = SelectedAssetClasses[SelectedAssetIndex];
+							if (AssetClass && (AssetClass->IsChildOf<UBlueprint>() || AssetClass->IsChildOf<UBlueprintGeneratedClass>()))
 							{
 								if (UE::Editor::Utilities::bEnableExperimentalScriptedActionsClassSupport)
 								{
 									FString ClassPathAsString;
-									Asset.GetTagValue(FBlueprintTags::GeneratedClassPath, ClassPathAsString);
+									if (AssetClass->IsChildOf<UBlueprintGeneratedClass>())
+									{
+										ClassPathAsString = Asset.GetObjectPathString();
+									}
+									else
+									{
+										Asset.GetTagValue(FBlueprintTags::GeneratedClassPath, ClassPathAsString);
+									}
 									FTopLevelAssetPath AssetClassPath(ClassPathAsString);
 
 									const IAssetRegistry& AssetRegistry = IAssetRegistry::GetChecked();
