Use the default tab id postfix (see UEditorUtilitySubsystem::GenerateNewTabID) in the "Run Utility Widget" and "Stop Utility Widget" buttons in the EUW editor.

This fixes an issue where the same UEditorUtilityWidgetBlueprint object was being erroneously registered twice with different tab ids by two different code paths.

diff --git a/Engine/Source/Editor/Blutility/Private/BlutilityUMGEditorExtensions.cpp b/Engine/Source/Editor/Blutility/Private/BlutilityUMGEditorExtensions.cpp
--- a/Engine/Source/Editor/Blutility/Private/BlutilityUMGEditorExtensions.cpp
+++ b/Engine/Source/Editor/Blutility/Private/BlutilityUMGEditorExtensions.cpp
@@ -60,7 +60,7 @@ public:
 										{
 											if (UEditorUtilityWidgetBlueprint* EditorWidget = Cast<UEditorUtilityWidgetBlueprint>(WidgetBlueprintEditor->GetBlueprintObj()))
 											{
-												GEditor->GetEditorSubsystem<UEditorUtilitySubsystem>()->SpawnAndRegisterTabWithId(EditorWidget, EditorWidget->GetFName());
+												GEditor->GetEditorSubsystem<UEditorUtilitySubsystem>()->SpawnAndRegisterTabWithId(EditorWidget, FName("_ActiveTab"));
 											}
 										}
 									}),
@@ -73,7 +73,7 @@ public:
 											if (UEditorUtilityWidgetBlueprint* EditorWidget = Cast<UEditorUtilityWidgetBlueprint>(WidgetBlueprintEditor->GetBlueprintObj()))
 											{
 												// Editor Utility Prepends path on Tab Name
-												FName UtilityTabName = FName(*(EditorWidget->GetPathName() + EditorWidget->GetFName().ToString()));
+												FName UtilityTabName = FName(*(EditorWidget->GetPathName() + FString("_ActiveTab")));
 												return !GEditor->GetEditorSubsystem<UEditorUtilitySubsystem>()->DoesTabExist(UtilityTabName);
 											}
 										}
@@ -96,7 +96,7 @@ public:
 											if (UEditorUtilityWidgetBlueprint* EditorWidget = Cast<UEditorUtilityWidgetBlueprint>(WidgetBlueprintEditor->GetBlueprintObj()))
 											{
 												// Editor Utility Prepends path on Tab Name
-												FName UtilityTabName = FName(*(EditorWidget->GetPathName() + EditorWidget->GetFName().ToString()));
+												FName UtilityTabName = FName(*(EditorWidget->GetPathName() + FString("_ActiveTab")));
 												GEditor->GetEditorSubsystem<UEditorUtilitySubsystem>()->CloseTabByID(UtilityTabName);
 											}
 										}
@@ -110,7 +110,7 @@ public:
 											if (UEditorUtilityWidgetBlueprint* EditorWidget = Cast<UEditorUtilityWidgetBlueprint>(WidgetBlueprintEditor->GetBlueprintObj()))
 											{
 												// Editor Utility Prepends path on Tab Name
-												FName UtilityTabName = FName(*(EditorWidget->GetPathName() + EditorWidget->GetFName().ToString()));
+												FName UtilityTabName = FName(*(EditorWidget->GetPathName() + FString("_ActiveTab")));
 												return GEditor->GetEditorSubsystem<UEditorUtilitySubsystem>()->DoesTabExist(UtilityTabName);
 											}
 										}
