Added GetTabIDFromBlueprint() to the EditorUtilitySubsystem so that an EUW instance can get its own TabID.
-Also updated the RegisterTabAndGetIDGeneratedClass code path to also set the TabID on the UEditorUtilityWidgetBlueprint, since that's how GetTabIDFromBlueprint() gets it.

diff --git a/Engine/Source/Editor/Blutility/Private/EditorUtilitySubsystem.cpp b/Engine/Source/Editor/Blutility/Private/EditorUtilitySubsystem.cpp
--- a/Engine/Source/Editor/Blutility/Private/EditorUtilitySubsystem.cpp
+++ b/Engine/Source/Editor/Blutility/Private/EditorUtilitySubsystem.cpp
@@ -389,7 +389,7 @@ void UEditorUtilitySubsystem::RegisterTabAndGetID(class UEditorUtilityWidgetBlue
 {
 	if (InBlueprint && !IsRunningCommandlet())
 	{
-		FName RegistrationName = NewTabID.IsNone() ? FName(*(InBlueprint->GetPathName() + LOCTEXT("ActiveTabSuffix", "_ActiveTab").ToString())) : FName(*(InBlueprint->GetPathName() + NewTabID.ToString()));
+		FName RegistrationName = GenerateNewTabID(InBlueprint, NewTabID);
 		FText DisplayName = InBlueprint->GetTabDisplayName();
 		IBlutilityModule& BlutilityModule = FModuleManager::GetModuleChecked<IBlutilityModule>("Blutility");
 		FLevelEditorModule& LevelEditorModule = FModuleManager::GetModuleChecked<FLevelEditorModule>(TEXT("LevelEditor"));
@@ -470,7 +470,7 @@ void UEditorUtilitySubsystem::RegisterTabAndGetIDGeneratedClass(UWidgetBlueprint
 {
 	if (InGeneratedWidgetBlueprint && !IsRunningCommandlet())
 	{
-		FName RegistrationName = NewTabID.IsNone() ? FName(*(InGeneratedWidgetBlueprint->GetPathName() + LOCTEXT("ActiveTabSuffix", "_ActiveTab").ToString())) : FName(*(InGeneratedWidgetBlueprint->GetPathName() + NewTabID.ToString()));
+		FName RegistrationName = GenerateNewTabID(InGeneratedWidgetBlueprint, NewTabID);
 		FText DisplayName;
 		const UEditorUtilityWidget* EditorUtilityWidget = InGeneratedWidgetBlueprint->GetDefaultObject<UEditorUtilityWidget>();
 
@@ -628,6 +628,17 @@ UEditorUtilityWidget* UEditorUtilitySubsystem::FindUtilityWidgetFromBlueprint(cl
 	return InBlueprint->GetCreatedWidget();
 }
 
+FName UEditorUtilitySubsystem::GetTabIDFromBlueprint(class UEditorUtilityWidgetBlueprint* InBlueprint)
+{
+	if (!IsValid(InBlueprint))
+	{ 
+		UE_LOG(LogEditorUtilityBlueprint, Error, TEXT("Found Invalid Blueprint in GetTabIDFromBlueprint"));
+		return NAME_None; 
+	}
+
+	return InBlueprint->GetRegistrationName();
+}
+
 bool UEditorUtilitySubsystem::Tick(float DeltaTime)
 {
 	UEditorUtilityTask* CurrentOrParentTask = GetActiveTask();
@@ -864,7 +875,7 @@ TSharedRef<SDockTab> UEditorUtilitySubsystem::SpawnEditorUITabFromGeneratedClass
 		.IsEnabled(FSlateApplication::Get().GetNormalExecutionAttribute());
 	UEditorUtilityWidget* CreatedUMGWidget = nullptr;
 
-	auto CreateUtilityWidgetFromGeneratedClass = [&CreatedUMGWidget](UWidgetBlueprintGeneratedClass* InGeneratedWidgetBlueprint)
+	auto CreateUtilityWidgetFromGeneratedClass = [&CreatedUMGWidget, &SpawnTabArgs](UWidgetBlueprintGeneratedClass* InGeneratedWidgetBlueprint)
 	{
 		TSharedRef<SWidget> TabWidget = SNullWidget::NullWidget;
 
@@ -890,6 +901,11 @@ TSharedRef<SDockTab> UEditorUtilitySubsystem::SpawnEditorUITabFromGeneratedClass
 				[
 					CreatedUMGWidget->TakeWidget()
 				];
+
+			if (UEditorUtilityWidgetBlueprint* WidgetBlueprint = Cast<UEditorUtilityWidgetBlueprint>(InGeneratedWidgetBlueprint->ClassGeneratedBy))
+			{
+				WidgetBlueprint->SetRegistrationName(SpawnTabArgs.GetTabId().TabType);
+			}
 		}
 		return TabWidget;
 	};
@@ -933,4 +949,10 @@ void UEditorUtilitySubsystem::OnMapChanged(UWorld* World, EMapChangeType MapChan
 	}
 }
 
+FName UEditorUtilitySubsystem::GenerateNewTabID(const UObject* WidgetObject, const FName& OptionalPostfix)
+{
+	static const FString DefaultPostfix = "_ActiveTab";
+	return OptionalPostfix.IsNone() ? FName(*(WidgetObject->GetPathName() + DefaultPostfix)) : FName(*(WidgetObject->GetPathName() + OptionalPostfix.ToString()));
+}
+
 #undef LOCTEXT_NAMESPACE
diff --git a/Engine/Source/Editor/Blutility/Public/EditorUtilitySubsystem.h b/Engine/Source/Editor/Blutility/Public/EditorUtilitySubsystem.h
--- a/Engine/Source/Editor/Blutility/Public/EditorUtilitySubsystem.h
+++ b/Engine/Source/Editor/Blutility/Public/EditorUtilitySubsystem.h
@@ -130,6 +130,10 @@ public:
 	UFUNCTION(BlueprintCallable, Category = "Development|Editor")
 	UE_API UEditorUtilityWidget* FindUtilityWidgetFromBlueprint(class UEditorUtilityWidgetBlueprint* InBlueprint);
 
+	/** Given an editor utility widget blueprint, get the TabID for it. This will return None if the widget is not currently in a tab.*/
+	UFUNCTION(BlueprintCallable, Category = "Development|Editor")
+	UE_API FName GetTabIDFromBlueprint(class UEditorUtilityWidgetBlueprint* InBlueprint);
+
 	/**  */
 	UFUNCTION(BlueprintCallable, Category = "Development|Editor")
 	UE_API void RegisterAndExecuteTask(UEditorUtilityTask* NewTask, UEditorUtilityTask* OptionalParentTask = nullptr);
@@ -183,6 +187,8 @@ protected:
 
 
 private:
+	FName GenerateNewTabID(const UObject* WidgetObject, const FName& OptionalPostfix);
+
 	IConsoleObject* RunTaskCommandObject = nullptr;
 	IConsoleObject* CancelAllTasksCommandObject = nullptr;
 	
