Blutility: exposing some public APIs for external use

#rnx

#Synced-CL 48389014

diff --git a/Engine/Source/Editor/Blutility/Private/BlutilityMenuExtensions.cpp b/Engine/Source/Editor/Blutility/Private/BlutilityMenuExtensions.cpp
--- a/Engine/Source/Editor/Blutility/Private/BlutilityMenuExtensions.cpp
+++ b/Engine/Source/Editor/Blutility/Private/BlutilityMenuExtensions.cpp
@@ -679,6 +679,28 @@ void FBlutilityMenuExtensions::CreateAssetBlutilityActionsMenu(FToolMenuSection&
 	);	
 }
 
+void FBlutilityMenuExtensions::CreateObjectBlutilityActionsMenu(
+	FMenuBuilder& MenuBuilder,
+	const FMenuParams& MenuParams,
+	const FObjectActionParams& ObjectActionParams)
+{
+	if (MenuParams.bAsSubMenu)
+	{
+		CreateBlutilityActionsMenu<UObject*>(MenuBuilder,
+			ObjectActionParams.Utils,
+			MenuParams.Label,
+			MenuParams.ToolTip,
+			ObjectActionParams.IsValidPropertyFunc,
+			ObjectActionParams.Objects,
+			MenuParams.IconName
+		);
+	}
+	else
+	{
+		CreateBlulityActionsSection<UObject*>(MenuBuilder, ObjectActionParams);
+	}
+}
+
 void FBlutilityMenuExtensions::OpenEditorForUtility(const FFunctionAndUtil& FunctionAndUtil)
 {
 	// Edit the script if we have shift held down
@@ -699,7 +721,7 @@ void FBlutilityMenuExtensions::OpenEditorForUtility(const FFunctionAndUtil& Func
 	}
 }
 
-void FBlutilityMenuExtensions::ExtractFunctions(TMap<TSharedRef<FAssetActionUtilityPrototype>, TSet<int32>>& Utils, TMap<FString, TArray<FFunctionAndUtil>>& OutCategoryFunctions)
+void FBlutilityMenuExtensions::ExtractFunctions(const TMap<TSharedRef<FAssetActionUtilityPrototype>, TSet<int32>>& Utils, TMap<FString, TArray<FFunctionAndUtil>>& OutCategoryFunctions)
 {
 	// Find the exposed functions available in each class, making sure to not list shared functions from a parent class more than once
 	for (TPair<TSharedRef<FAssetActionUtilityPrototype>, TSet<int32>> UtilitySelectionPair : Utils)
@@ -946,6 +968,227 @@ void FBlutilityMenuExtensions::CreateBlutilityActionsMenu(FMenuBuilder& MenuBuil
 	}
 }
 
+template<typename SelectionType>
+void FBlutilityMenuExtensions::CreateBlulityActionsSection(FMenuBuilder& InMenuBuilder, const FObjectActionParams& InObjectActionParams)
+{
+	TMap<FString, TArray<FFunctionAndUtil>> CategoryFunctions;
+	ExtractFunctions(InObjectActionParams.Utils, CategoryFunctions);
+	
+	auto AddFunctionEntries = [Selection = InObjectActionParams.Objects, IsValidPropertyType = InObjectActionParams.IsValidPropertyFunc](
+		const TArray<FFunctionAndUtil>& FunctionUtils)
+	{
+		BlutilityUtil::FTextFilterExpressionContext TextFilterContext;
+		
+		FTextFilterExpressionEvaluator TextFilterExpressionEvaluator(ETextFilterExpressionEvaluatorMode::Complex);
+		TArray<FMenuEntryParams> GeneratedMenuEntryParams;
+		for (const FFunctionAndUtil& FunctionAndUtil : FunctionUtils)
+		{
+			bool PassesFilterCondition = true;
+			bool bShowInMenu = true;
+			FString FilterFailureMessage;
+			
+			if constexpr ( std::is_same_v<FAssetData, SelectionType> )
+			{
+				TArray<FAssetActionSupportCondition> Conditions = FunctionAndUtil.Util->GetAssetActionSupportConditions();
+
+				for (const FAssetActionSupportCondition& Condition : Conditions)
+				{
+					TextFilterExpressionEvaluator.SetFilterText(FText::FromString(Condition.Filter));
+				
+					for (const int32& SelectionIndex : FunctionAndUtil.SelectionIndices)
+					{
+						const auto SelectedAsset = Cast<SelectionType>(Selection[SelectionIndex]);
+						TextFilterContext.SetAsset(&SelectedAsset);
+						if (!TextFilterExpressionEvaluator.TestTextFilter(TextFilterContext))
+						{
+							PassesFilterCondition = false;
+							FilterFailureMessage = Condition.FailureReason;
+							break;
+						}
+					}
+
+					if (!PassesFilterCondition)
+					{
+						bShowInMenu = Condition.bShowInMenuIfFilterFails;
+						break;
+					}
+				}
+			}
+
+			if (!PassesFilterCondition && !bShowInMenu)
+			{
+				continue;
+			}
+			
+			FText TooltipText;
+
+			if (FunctionAndUtil.Util->GetUtilityBlueprintAsset().AssetClassPath == UBlueprintGeneratedClass::StaticClass()->GetClassPathName())
+			{
+				TooltipText = FText::Format(LOCTEXT("AssetUtilTooltip", "{0}\n\n(Click to execute)"), FunctionAndUtil.FunctionData.TooltipText);
+			}
+			else if (FilterFailureMessage.IsEmpty())
+			{
+				TooltipText = FText::Format(LOCTEXT("AssetUtilTooltipFormat", "{0}\n\n(Shift-click to edit script)"), FunctionAndUtil.FunctionData.TooltipText);	
+			}
+			else
+			{
+				TooltipText = FText::Format(LOCTEXT("AssetUtilTooltipWithErrorFormat", "{0}\n\n({1})\n\n(Shift-click to edit script)"), FunctionAndUtil.FunctionData.TooltipText, FText::FromString(FilterFailureMessage));
+			}
+
+			FMenuEntryParams MenuEntryParams;
+			MenuEntryParams.LabelOverride = FunctionAndUtil.FunctionData.NameText;
+			MenuEntryParams.ToolTipOverride = TooltipText;
+			MenuEntryParams.IconOverride = FSlateIcon(FAppStyle::GetAppStyleSetName(), "GraphEditor.Event_16x");
+			MenuEntryParams.DirectActions = FUIAction(FExecuteAction::CreateLambda([FunctionAndUtil, Selection, IsValidPropertyType]
+				{
+					if (FSlateApplication::Get().GetModifierKeys().IsShiftDown())
+					{
+						OpenEditorForUtility(FunctionAndUtil);
+					}
+					else
+					{
+						// We dont run this on the CDO, as bad things could occur!
+						UObject* UtilityObject = Cast<UObject>(FunctionAndUtil.Util->LoadUtilityAsset());
+						UFunction* Function = FunctionAndUtil.GetFunction();
+						if (UtilityObject && Function)
+						{
+							UObject* TempObject = NewObject<UObject>(GetTransientPackage(), UtilityObject->GetClass());
+							TempObject->AddToRoot(); // Some Blutility actions might run GC so the TempObject needs to be rooted to avoid getting destroyed
+
+							if (Function->NumParms > 0)
+							{
+								// Create a parameter struct and fill in defaults
+								TSharedRef<FStructOnScope> FuncParams = MakeShared<FStructOnScope>(Function);
+
+								FProperty* FirstParamProperty = nullptr;
+
+								int32 ParameterIndex = 0;
+								for (TFieldIterator<FProperty> It(Function); It&& It->HasAnyPropertyFlags(CPF_Parm); ++It)
+								{
+									FString Defaults;
+									if (UEdGraphSchema_K2::FindFunctionParameterDefaultValue(Function, *It, Defaults))
+									{
+										It->ImportText_Direct(*Defaults, It->ContainerPtrToValuePtr<uint8>(FuncParams->GetStructMemory()), nullptr, PPF_None);
+									}
+
+									// Check to see if the first parameter matches the selection object type, in that case we can directly forward the selection to it
+									if (ParameterIndex == 0 && IsValidPropertyType(*It))
+									{
+										FirstParamProperty = *It;
+									}
+
+									++ParameterIndex;
+							   }
+
+								bool bApply = true;
+
+								if (!FirstParamProperty || ParameterIndex > 1)
+								{
+									// pop up a dialog to input params to the function
+									TSharedRef<SWindow> Window = SNew(SWindow)
+										.Title(Function->GetDisplayNameText())
+										.ClientSize(FVector2D(400, 200))
+										.SupportsMinimize(false)
+										.SupportsMaximize(false);
+
+									TSharedPtr<SFunctionParamDialog> Dialog;
+									Window->SetContent(
+										SAssignNew(Dialog, SFunctionParamDialog, Window, FuncParams, FirstParamProperty ? FirstParamProperty->GetFName() : NAME_None)
+										.OkButtonText(LOCTEXT("OKButton", "OK"))
+										.OkButtonTooltipText(Function->GetToolTipText()));
+
+									GEditor->EditorAddModalWindow(Window);
+									bApply = Dialog->bOKPressed;
+								}
+
+								if (bApply)
+								{
+									FScopedTransaction Transaction(NSLOCTEXT("UnrealEd", "BlutilityAction", "Blutility Action"));
+									FEditorScriptExecutionGuard ScriptGuard;
+									const bool bForwardUserSelection = FirstParamProperty != nullptr;
+									if (bForwardUserSelection)
+									{
+										// For each user-select asset forward the selection object into the function first's parameter (if it matches)
+										const FString Path = FirstParamProperty->GetPathName(Function);
+
+										// Ensure we only process selection objects that are valid for this function/utility
+										for (const int32& SelectionIndex : FunctionAndUtil.SelectionIndices)
+										{
+											const auto SelectedAsset = Selection[SelectionIndex];
+											FirstParamProperty->CopySingleValue(
+												FirstParamProperty->ContainerPtrToValuePtr<uint8>(FuncParams->GetStructMemory()), &SelectedAsset);
+											TempObject->ProcessEvent(Function, FuncParams->GetStructMemory());
+										}
+									}
+									else
+									{
+										// User is expected to manage the asset selection on its own
+										TempObject->ProcessEvent(Function, FuncParams->GetStructMemory());
+									}
+								}
+							}
+							else
+							{
+								FScopedTransaction Transaction(NSLOCTEXT("UnrealEd", "BlutilityAction", "Blutility Action"));
+								FEditorScriptExecutionGuard ScriptGuard;
+								TempObject->ProcessEvent(Function, nullptr);
+							}
+							
+							TempObject->RemoveFromRoot();
+						}
+					}
+				}),
+				FCanExecuteAction::CreateLambda([PassesFilterCondition]()
+				{
+					return PassesFilterCondition;
+				})
+			);
+
+			GeneratedMenuEntryParams.Emplace(MenuEntryParams);
+		}
+
+		return GeneratedMenuEntryParams;
+	};
+
+	// Add a menu item for each function
+	if (CategoryFunctions.Num() > 0)
+	{
+		TArray<FString> CategoryNames;
+		CategoryFunctions.GenerateKeyArray(CategoryNames);
+		CategoryNames.Remove(FString());
+		CategoryNames.Sort();
+			
+		for (const FString& CategoryName : CategoryNames)
+		{
+			const TArray<FMenuEntryParams> GeneratedMenuEntries = AddFunctionEntries(CategoryFunctions.FindChecked(CategoryName));
+			if (!GeneratedMenuEntries.IsEmpty())
+			{
+				InMenuBuilder.BeginSection(*CategoryName, FText::FromString(CategoryName));
+				{
+					for (const FMenuEntryParams& MenuParams : GeneratedMenuEntries)
+					{
+						InMenuBuilder.AddMenuEntry(MenuParams);
+					}
+				}
+				InMenuBuilder.EndSection();
+			}
+		}
+
+		// Non-categorized functions
+		if (const TArray<FFunctionAndUtil>* DefaultCategoryFunctionsPtr = CategoryFunctions.Find(FString()))
+		{
+			const TArray<FMenuEntryParams> GeneratedMenuEntries = AddFunctionEntries(*DefaultCategoryFunctionsPtr);
+			if (!GeneratedMenuEntries.IsEmpty())
+			{
+				for (const FMenuEntryParams& MenuParams : GeneratedMenuEntries)
+				{
+					InMenuBuilder.AddMenuEntry(MenuParams);
+				}
+			}
+		}
+	}
+}
+
 template<typename SelectionType>
 void FBlutilityMenuExtensions::CreateBlutilityActionsMenu(FToolMenuSection& InSection, TMap<TSharedRef<FAssetActionUtilityPrototype>, TSet<int32>> Utils, const FName& MenuName, const FText& MenuLabel, const FText& MenuToolTip, TFunction<bool(const FProperty * Property)> IsValidPropertyType, const TArray<SelectionType> Selection, const FName& IconName)
 {
diff --git a/Engine/Source/Editor/Blutility/Private/BlutilityMenuExtensions.h b/Engine/Source/Editor/Blutility/Private/BlutilityMenuExtensions.h
deleted file mode 100644
diff --git a/Engine/Source/Editor/Blutility/Private/EditorUtilityAssetPrototype.h b/Engine/Source/Editor/Blutility/Private/EditorUtilityAssetPrototype.h
deleted file mode 100644
diff --git a/Engine/Source/Editor/Blutility/Public/BlutilityMenuExtensions.h b/Engine/Source/Editor/Blutility/Public/BlutilityMenuExtensions.h
new file mode 100644
--- /dev/null
+++ b/Engine/Source/Editor/Blutility/Public/BlutilityMenuExtensions.h
@@ -0,0 +1,111 @@
+// Copyright Epic Games, Inc. All Rights Reserved.
+
+#pragma once
+
+#include "EditorUtilityAssetPrototype.h"
+#include "Containers/Array.h"
+#include "Containers/Map.h"
+#include "Containers/Set.h"
+#include "Containers/SparseArray.h"
+#include "HAL/Platform.h"
+#include "Templates/Function.h"
+#include "UObject/Interface.h"
+#include "UObject/NameTypes.h"
+#include "UObject/ObjectMacros.h"
+#include "UObject/TopLevelAssetPath.h"
+#include "UObject/UObjectGlobals.h"
+
+class AActor;
+class FMenuBuilder;
+class FProperty;
+class FString;
+class FText;
+class IEditorUtilityExtension;
+struct FToolMenuSection;
+class UFunction;
+class UObject;
+class FAssetActionUtilityPrototype;
+
+#define UE_API BLUTILITY_API
+
+// Blutility Menu extension helpers
+class FBlutilityMenuExtensions
+{
+public:
+	/** Helper function to get all Blutility classes derived from the specified class name */
+	static UE_API void GetBlutilityClasses(TArray<FAssetData>& OutAssets, FTopLevelAssetPath InClassName);
+
+	/** Helper function that populates a menu based on the exposed functions in a set of Blutility objects */
+	static void CreateActorBlutilityActionsMenu(FToolMenuSection& InSection, TMap<TSharedRef<FAssetActionUtilityPrototype>, TSet<int32>> Utils, const TArray<AActor*> SelectedSupportedActors);
+	static void CreateAssetBlutilityActionsMenu(FToolMenuSection& InSection, TMap<TSharedRef<FAssetActionUtilityPrototype>, TSet<int32>> Utils, const TArray<FAssetData> SelectedSupportedAssets);
+
+	/* Structure that defines parameters related to the menu passed into the CreateObjectBlutilityActionsMenu function. */
+	struct FMenuParams
+	{
+		FText Label;
+		FText ToolTip;
+		FName IconName = "GraphEditor.Event_16x";
+		bool bAsSubMenu = true;
+	};
+		
+	/* Structure that defines the parameters related to actions and the objects to act on passed into the CreateObjectBlutilityActionsMenu function. */
+	struct FObjectActionParams
+	{
+		TArray<UObject*> Objects;
+		TMap<TSharedRef<FAssetActionUtilityPrototype>, TSet<int32>> Utils;
+		TFunction<bool(const FProperty* Property)> IsValidPropertyFunc = [](const FProperty*){ return true; };
+	};
+	
+	/** Helper function that populates a menu based on the exposed functions in a set of Blutility objects */
+	static UE_API void CreateObjectBlutilityActionsMenu(
+		FMenuBuilder& MenuBuilder,
+		const FMenuParams& MenuParams,
+		const FObjectActionParams& ObjectActionParams
+	);
+	
+protected:
+	// Helper struct to track the util to call a function on
+	struct FFunctionAndUtil
+	{
+		FFunctionAndUtil(const FBlutilityFunctionData& InFunctionData, const TSharedRef<FAssetActionUtilityPrototype>& InUtil, TSet<int32>& InSelection)
+			: FunctionData(InFunctionData)
+			, Util(InUtil) 
+			, SelectionIndices(InSelection)
+		{}
+
+		bool operator==(const FFunctionAndUtil& InFunction) const
+		{
+			return FunctionData == InFunction.FunctionData && Util == InFunction.Util;
+		}
+
+		UFunction* GetFunction() const
+		{
+			if (const UObject* UtilityCDO = Util->LoadUtilityAsset())
+			{
+				if (const UClass* UtilityClass = UtilityCDO->GetClass())
+				{
+					return UtilityClass->FindFunctionByName(FunctionData.Name);
+				}
+			}
+
+			return nullptr;
+		}
+
+		FBlutilityFunctionData FunctionData;
+		TSharedRef<FAssetActionUtilityPrototype> Util;
+		/** Indices to original object selection array which are supported by this utility */
+		TSet<int32> SelectionIndices;
+	};
+
+protected:
+	template<typename SelectionType>
+	static void CreateBlutilityActionsMenu(FMenuBuilder& MenuBuilder, TMap<TSharedRef<FAssetActionUtilityPrototype>, TSet<int32>> Utils, const FText& MenuLabel, const FText& MenuToolTip, TFunction<bool(const FProperty* Property)> IsValidPropertyType, const TArray<SelectionType> Selection, const FName& IconName = "GraphEditor.Event_16x");
+	template<typename SelectionType>
+	static void CreateBlutilityActionsMenu(FToolMenuSection& InSection, TMap<TSharedRef<FAssetActionUtilityPrototype>, TSet<int32>> Utils, const FName& MenuName, const FText& MenuLabel, const FText& MenuToolTip, TFunction<bool(const FProperty* Property)> IsValidPropertyType, const TArray<SelectionType> Selection, const FName& IconName = "GraphEditor.Event_16x");
+	template<typename SelectionType>
+	static void CreateBlulityActionsSection(FMenuBuilder& InMenuBuilder, const FObjectActionParams& InObjectActionParams);
+	static void OpenEditorForUtility(const FFunctionAndUtil& FunctionAndUtil);
+	static void ExtractFunctions(const TMap<TSharedRef<FAssetActionUtilityPrototype>, TSet<int32>>& Utils, TMap<FString, TArray<FFunctionAndUtil>>& OutCategoryFunctions);
+};
+
+#undef UE_API
\ No newline at end of file
diff --git a/Engine/Source/Editor/Blutility/Public/EditorUtilityAssetPrototype.h b/Engine/Source/Editor/Blutility/Public/EditorUtilityAssetPrototype.h
new file mode 100644
--- /dev/null
+++ b/Engine/Source/Editor/Blutility/Public/EditorUtilityAssetPrototype.h
@@ -0,0 +1,100 @@
+// Copyright Epic Games, Inc. All Rights Reserved.
+
+#pragma once
+
+#include "AssetActionUtility.h"
+#include "UObject/Object.h"
+#include "AssetRegistry/AssetData.h"
+#include "IEditorUtilityExtension.h"
+
+#include "EditorUtilityAssetPrototype.generated.h"
+
+USTRUCT()
+struct FBlutilityFunctionData
+{
+	GENERATED_BODY()
+
+public:
+	UPROPERTY()
+	TSoftClassPtr<UClass> Class;
+	
+	UPROPERTY()
+	FName Name;
+	
+	UPROPERTY()
+	FText NameText;
+
+	UPROPERTY()
+	FString Category;
+
+	UPROPERTY()
+	FText TooltipText;
+
+	bool operator !=(const FBlutilityFunctionData& Other) const
+	{
+		return !(*this == Other);
+	}
+
+	bool operator == (const FBlutilityFunctionData& Other) const
+	{
+		return Class == Other.Class && Name == Other.Name;
+	}
+};
+
+namespace AssetActionUtilityTags
+{
+	extern const FName BlutilityTagVersion;
+	extern const FName SupportedClasses;
+	extern const FName IsActionForBlueprint;
+	extern const FName CallableFunctions;
+	
+	extern const int32 TagVersion;
+}
+
+#define UE_API BLUTILITY_API
+
+class FAssetActionUtilityPrototype
+{
+public:
+	FAssetActionUtilityPrototype(const FAssetData& SourceAsset)
+		: UtilityBlueprintAsset(SourceAsset)
+	{
+	}
+
+	UObject* LoadUtilityAsset() const;
+
+	UE_API bool IsLatestVersion() const;
+	bool AreSupportedClassesForBlueprints() const;
+	UE_API TArray<TSoftClassPtr<UObject>> GetSupportedClasses() const;
+	TArray<FAssetActionSupportCondition> GetAssetActionSupportConditions() const;
+	UE_API TArray<FBlutilityFunctionData> GetCallableFunctions() const;
+
+	const FAssetData& GetUtilityBlueprintAsset() const { return UtilityBlueprintAsset; }
+
+	bool operator !=(const FAssetActionUtilityPrototype& Other) const
+	{
+		return !(*this == Other);
+	}
+
+	bool operator == (const FAssetActionUtilityPrototype& Other) const
+	{
+		return UtilityBlueprintAsset.PackageName == Other.UtilityBlueprintAsset.PackageName;
+	}
+
+public:
+	static void AddTagsFor_Version(FAssetRegistryTagsContext Context);
+	static void AddTagsFor_SupportedClasses(const TArray<TSoftClassPtr<UObject>>& SupportedClasses, FAssetRegistryTagsContext Context);
+	static void AddTagsFor_SupportedConditions(const TArray<FAssetActionSupportCondition>& SupportedConditions, FAssetRegistryTagsContext Context);
+	static void AddTagsFor_IsActionForBlueprints(bool IsActionForBlueprints, FAssetRegistryTagsContext Context);
+	static void AddTagsFor_CallableFunctions(const UObject* FunctionsSource, FAssetRegistryTagsContext Context);
+
+private:
+	FAssetData UtilityBlueprintAsset;
+};
+
+FORCEINLINE uint32 GetTypeHash(const FAssetActionUtilityPrototype& ExtensionAsset)
+{
+	return HashCombine(GetTypeHash(ExtensionAsset.GetUtilityBlueprintAsset().PackageName), GetTypeHash(ExtensionAsset.GetUtilityBlueprintAsset().AssetName));
+}
+
+#undef UE_API
\ No newline at end of file
