Allow user to select object

#rnx
#Synced-CL 50672290

diff --git a/Engine/Source/Editor/Blutility/Classes/AssetActionUtility.h b/Engine/Source/Editor/Blutility/Classes/AssetActionUtility.h
--- a/Engine/Source/Editor/Blutility/Classes/AssetActionUtility.h
+++ b/Engine/Source/Editor/Blutility/Classes/AssetActionUtility.h
@@ -105,7 +105,7 @@ protected:
 	 * utility class.
 	 * @note When bIsActionForBlueprints is true, this will compare against the generated class of any Blueprint assets.
 	 */
-	UPROPERTY(EditDefaultsOnly, Category="Asset Support", meta=(AllowAbstract))
+	UPROPERTY(EditDefaultsOnly, Category="Asset Support", meta=(AllowAbstract, ShowObjectRootClass))
 	TArray<TSoftClassPtr<UObject>> SupportedClasses;
 
 	/**
