[Blutility] Add modal Editor Utility Dialog system

Adds UEditorUtilityDialogWidget (base class for dialog-hosted EUWs),
SEditorUtilityDialog (Slate window), and UEditorUtilityDialogLibrary
(BP function library for showing modal/non-modal EUW dialogs).

Supports three workflows:
- Class-based: pass a widget class, dialog creates and shows it
- Instance via Create: create widget, set data, then show
- Instance via Construct Object: construct, bind OnOpened, show

https://horde.devtools.epicgames.com/job/69bb007f322cd371995deb80

#rnx
#tests tested in Editor with different setups.
#Synced-CL 51944695

diff --git a/Engine/Source/Editor/Blutility/Private/EditorUtilityDialogLibrary.cpp b/Engine/Source/Editor/Blutility/Private/EditorUtilityDialogLibrary.cpp
new file mode 100644
--- /dev/null
+++ b/Engine/Source/Editor/Blutility/Private/EditorUtilityDialogLibrary.cpp
@@ -0,0 +1,211 @@
+// Copyright Epic Games, Inc. All Rights Reserved.
+
+#include "EditorUtilityDialogLibrary.h"
+#include "SEditorUtilityDialog.h"
+#include "EditorUtilityDialogWidget.h"
+
+#include "EditorUtilityCommon.h"
+#include "Editor.h"
+#include "Blueprint/UserWidget.h"
+
+FEditorUtilityDialogResult UEditorUtilityDialogLibrary::ShowModalEditorUtilityDialog(
+	TSubclassOf<UEditorUtilityDialogWidget> WidgetClass,
+	FText Title,
+	EAppMsgType::Type ButtonLayout,
+	FVector2D Size)
+{
+	FEditorUtilityDialogResult Result;
+
+	TPair<UEditorUtilityDialogWidget*, TSharedPtr<SEditorUtilityDialog>> Created = CreateDialogInternal(WidgetClass, Title, ButtonLayout, Size);
+	UEditorUtilityDialogWidget* Widget = Created.Key;
+	TSharedPtr<SEditorUtilityDialog> Dialog = Created.Value;
+	if (!Widget || !Dialog.IsValid())
+	{
+		return Result;
+	}
+
+	Result.WidgetInstance = Widget;
+
+	Widget->OnDialogOpened();
+
+	// Show modal - blocks until closed. SEditorUtilityDialog::ShowModal handles
+	// unattended mode internally (fires widget close events and returns Cancel).
+	Result.ReturnType = Dialog->ShowModal();
+
+	return Result;
+}
+
+UEditorUtilityDialogWidget* UEditorUtilityDialogLibrary::ShowEditorUtilityDialog(
+	TSubclassOf<UEditorUtilityDialogWidget> WidgetClass,
+	const FOnEditorUtilityDialogClosed& OnClosed,
+	FText Title,
+	EAppMsgType::Type ButtonLayout,
+	FVector2D Size)
+{
+	TPair<UEditorUtilityDialogWidget*, TSharedPtr<SEditorUtilityDialog>> Created = CreateDialogInternal(WidgetClass, Title, ButtonLayout, Size);
+	UEditorUtilityDialogWidget* Widget = Created.Key;
+	TSharedPtr<SEditorUtilityDialog> Dialog = Created.Value;
+	if (!Widget || !Dialog.IsValid())
+	{
+		return nullptr;
+	}
+
+	Widget->OnDialogOpened();
+
+	// Show non-modal - returns immediately.
+	// The close delegate fires from SEditorUtilityDialog::HandleWindowClosed for all close paths.
+	Dialog->Show(FOnEditorUtilityDialogSlateClosedNative::CreateLambda(
+		[OnClosed, Widget](EAppReturnType::Type Result)
+		{
+			OnClosed.ExecuteIfBound(Result, Widget);
+		}));
+
+	return Widget;
+}
+
+UEditorUtilityDialogWidget* UEditorUtilityDialogLibrary::CreateEditorUtilityDialogWidget(
+	TSubclassOf<UEditorUtilityDialogWidget> WidgetClass)
+{
+	if (!WidgetClass)
+	{
+		UE_LOG(LogEditorUtilityBlueprint, Error, TEXT("CreateEditorUtilityDialogWidget: WidgetClass is null."));
+		return nullptr;
+	}
+
+	UWorld* World = GEditor ? GEditor->GetEditorWorldContext().World() : nullptr;
+	if (!World)
+	{
+		UE_LOG(LogEditorUtilityBlueprint, Error, TEXT("CreateEditorUtilityDialogWidget: No editor world available."));
+		return nullptr;
+	}
+
+	// CreateWidget calls Initialize() which builds the widget tree,
+	// so child widget references (text blocks, etc.) are valid after this.
+	UEditorUtilityDialogWidget* Widget = CreateWidget<UEditorUtilityDialogWidget>(World, WidgetClass);
+	if (!Widget)
+	{
+		UE_LOG(LogEditorUtilityBlueprint, Error, TEXT("CreateEditorUtilityDialogWidget: Failed to create widget of class %s."), *WidgetClass->GetName());
+		return nullptr;
+	}
+
+	Widget->SetFlags(RF_Transient);
+	return Widget;
+}
+
+FEditorUtilityDialogResult UEditorUtilityDialogLibrary::ShowModalEditorUtilityDialogInstance(
+	UEditorUtilityDialogWidget* WidgetInstance,
+	FText Title,
+	EAppMsgType::Type ButtonLayout,
+	FVector2D Size)
+{
+	FEditorUtilityDialogResult Result;
+
+	if (!WidgetInstance)
+	{
+		UE_LOG(LogEditorUtilityBlueprint, Error, TEXT("ShowModalEditorUtilityDialogInstance: WidgetInstance is null."));
+		return Result;
+	}
+
+	TSharedPtr<SEditorUtilityDialog> Dialog = WrapWidgetInDialog(WidgetInstance, Title, ButtonLayout, Size);
+	if (!Dialog.IsValid())
+	{
+		return Result;
+	}
+
+	Result.WidgetInstance = WidgetInstance;
+
+	WidgetInstance->OnDialogOpened();
+
+	Result.ReturnType = Dialog->ShowModal();
+
+	return Result;
+}
+
+void UEditorUtilityDialogLibrary::ShowEditorUtilityDialogInstance(
+	UEditorUtilityDialogWidget* WidgetInstance,
+	const FOnEditorUtilityDialogClosed& OnClosed,
+	FText Title,
+	EAppMsgType::Type ButtonLayout,
+	FVector2D Size)
+{
+	if (!WidgetInstance)
+	{
+		UE_LOG(LogEditorUtilityBlueprint, Error, TEXT("ShowEditorUtilityDialogInstance: WidgetInstance is null."));
+		return;
+	}
+
+	TSharedPtr<SEditorUtilityDialog> Dialog = WrapWidgetInDialog(WidgetInstance, Title, ButtonLayout, Size);
+	if (!Dialog.IsValid())
+	{
+		return;
+	}
+
+	WidgetInstance->OnDialogOpened();
+
+	Dialog->Show(FOnEditorUtilityDialogSlateClosedNative::CreateLambda(
+		[OnClosed, WidgetInstance](EAppReturnType::Type Result)
+		{
+			OnClosed.ExecuteIfBound(Result, WidgetInstance);
+		}));
+}
+
+TPair<UEditorUtilityDialogWidget*, TSharedPtr<SEditorUtilityDialog>> UEditorUtilityDialogLibrary::CreateDialogInternal(
+	TSubclassOf<UEditorUtilityDialogWidget> WidgetClass,
+	const FText& Title,
+	EAppMsgType::Type ButtonLayout,
+	const FVector2D& Size)
+{
+	if (!WidgetClass)
+	{
+		UE_LOG(LogEditorUtilityBlueprint, Error, TEXT("ShowEditorUtilityDialog: WidgetClass is null."));
+		return TPair<UEditorUtilityDialogWidget*, TSharedPtr<SEditorUtilityDialog>>(nullptr, nullptr);
+	}
+
+	UWorld* World = GEditor ? GEditor->GetEditorWorldContext().World() : nullptr;
+	if (!World)
+	{
+		UE_LOG(LogEditorUtilityBlueprint, Error, TEXT("ShowEditorUtilityDialog: No editor world available."));
+		return TPair<UEditorUtilityDialogWidget*, TSharedPtr<SEditorUtilityDialog>>(nullptr, nullptr);
+	}
+
+	UEditorUtilityDialogWidget* Widget = CreateWidget<UEditorUtilityDialogWidget>(World, WidgetClass);
+	if (!Widget)
+	{
+		UE_LOG(LogEditorUtilityBlueprint, Error, TEXT("ShowEditorUtilityDialog: Failed to create widget of class %s."), *WidgetClass->GetName());
+		return TPair<UEditorUtilityDialogWidget*, TSharedPtr<SEditorUtilityDialog>>(nullptr, nullptr);
+	}
+
+	Widget->SetFlags(RF_Transient);
+
+	TSharedPtr<SEditorUtilityDialog> Dialog = WrapWidgetInDialog(Widget, Title, ButtonLayout, Size);
+	return TPair<UEditorUtilityDialogWidget*, TSharedPtr<SEditorUtilityDialog>>(Widget, Dialog);
+}
+
+TSharedPtr<SEditorUtilityDialog> UEditorUtilityDialogLibrary::WrapWidgetInDialog(
+	UEditorUtilityDialogWidget* Widget,
+	const FText& Title,
+	EAppMsgType::Type ButtonLayout,
+	const FVector2D& Size)
+{
+	check(Widget);
+
+	// "Construct Object of Class" in Blueprint creates the UObject but does NOT call
+	// UUserWidget::Initialize(), which builds the widget tree. Without it, TakeWidget()
+	// returns an empty Slate widget. CreateWidget<> calls Initialize() automatically,
+	// but the instance-based path needs this explicit call.
+	// Initialize() is idempotent (checks bInitialized internally), so safe to call always.
+	Widget->Initialize();
+
+	const FText EffectiveTitle = Title.IsEmpty() ? Widget->DialogTitle : Title;
+	const FVector2D EffectiveSize = Size.IsNearlyZero() ? Widget->DesiredDialogSize : Size;
+
+	TSharedRef<SEditorUtilityDialog> Dialog = SNew(SEditorUtilityDialog)
+		.Title(EffectiveTitle)
+		.ButtonLayout(ButtonLayout)
+		.ClientSize(EffectiveSize);
+
+	Dialog->SetContent(Widget);
+	Widget->SetOwningDialog(Dialog);
+
+	return Dialog;
+}
diff --git a/Engine/Source/Editor/Blutility/Private/EditorUtilityDialogWidget.cpp b/Engine/Source/Editor/Blutility/Private/EditorUtilityDialogWidget.cpp
new file mode 100644
--- /dev/null
+++ b/Engine/Source/Editor/Blutility/Private/EditorUtilityDialogWidget.cpp
@@ -0,0 +1,17 @@
+// Copyright Epic Games, Inc. All Rights Reserved.
+
+#include "EditorUtilityDialogWidget.h"
+#include "SEditorUtilityDialog.h"
+
+void UEditorUtilityDialogWidget::CloseDialog(EAppReturnType::Type Result)
+{
+	if (TSharedPtr<SEditorUtilityDialog> Dialog = OwningDialog.Pin())
+	{
+		Dialog->CloseWithResult(Result);
+	}
+}
+
+void UEditorUtilityDialogWidget::SetOwningDialog(TSharedPtr<SEditorUtilityDialog> InDialog)
+{
+	OwningDialog = InDialog;
+}
diff --git a/Engine/Source/Editor/Blutility/Private/SEditorUtilityDialog.cpp b/Engine/Source/Editor/Blutility/Private/SEditorUtilityDialog.cpp
new file mode 100644
--- /dev/null
+++ b/Engine/Source/Editor/Blutility/Private/SEditorUtilityDialog.cpp
@@ -0,0 +1,331 @@
+// Copyright Epic Games, Inc. All Rights Reserved.
+
+#include "SEditorUtilityDialog.h"
+#include "EditorUtilityDialogWidget.h"
+
+#include "Dialog/DialogCommands.h"
+#include "Framework/Application/SlateApplication.h"
+#include "Framework/Docking/TabManager.h"
+#include "Misc/App.h"
+#include "Styling/AppStyle.h"
+#include "Widgets/Input/SButton.h"
+#include "Widgets/Layout/SBox.h"
+#include "Widgets/SBoxPanel.h"
+#include "Widgets/SNullWidget.h"
+#include "Widgets/Text/STextBlock.h"
+#include "EditorUtilityCommon.h"
+#include "Blueprint/UserWidget.h"
+
+#define LOCTEXT_NAMESPACE "SEditorUtilityDialog"
+
+void SEditorUtilityDialog::Construct(const FArguments& InArgs)
+{
+	const EAppMsgType::Type MsgType = InArgs._ButtonLayout;
+	EAppReturnType::Type ConfirmResult = EAppReturnType::Ok;
+
+	switch (MsgType)
+	{
+	case EAppMsgType::Ok:
+		ConfirmResult = EAppReturnType::Ok;
+		CancelResult = EAppReturnType::Ok;
+		break;
+	case EAppMsgType::YesNo:
+		ConfirmResult = EAppReturnType::Yes;
+		CancelResult = EAppReturnType::No;
+		break;
+	case EAppMsgType::OkCancel:
+		ConfirmResult = EAppReturnType::Ok;
+		CancelResult = EAppReturnType::Cancel;
+		break;
+	case EAppMsgType::YesNoCancel:
+		ConfirmResult = EAppReturnType::Yes;
+		CancelResult = EAppReturnType::Cancel;
+		break;
+	case EAppMsgType::CancelRetryContinue:
+		ConfirmResult = EAppReturnType::Continue;
+		CancelResult = EAppReturnType::Cancel;
+		break;
+	case EAppMsgType::YesNoYesAllNoAll:
+		ConfirmResult = EAppReturnType::Yes;
+		CancelResult = EAppReturnType::No;
+		break;
+	case EAppMsgType::YesNoYesAllNoAllCancel:
+		ConfirmResult = EAppReturnType::Yes;
+		CancelResult = EAppReturnType::Cancel;
+		break;
+	case EAppMsgType::YesNoYesAll:
+		ConfirmResult = EAppReturnType::Yes;
+		CancelResult = EAppReturnType::No;
+		break;
+	default:
+		break;
+	}
+
+	// Store the confirm result for hotkey capture. CancelResult is already a member.
+	const EAppReturnType::Type LocalConfirmResult = ConfirmResult;
+
+	// Build the command list for hotkeys (Escape = Cancel/No, Enter = Confirm/Ok/Yes).
+	// CanExecuteHotkey() prevents firing after the dialog has begun closing (matching SCustomDialog pattern).
+	CommandList = MakeShared<FUICommandList>();
+	CommandList->MapAction(FDialogCommands::Get().Cancel,
+		FUIAction(
+			FSimpleDelegate::CreateSPLambda(this, [this]() { HandleButtonResult(CancelResult); }),
+			FCanExecuteAction::CreateSP(this, &SEditorUtilityDialog::CanExecuteHotkey)
+		)
+	);
+	CommandList->MapAction(FDialogCommands::Get().Confirm,
+		FUIAction(
+			FSimpleDelegate::CreateSPLambda(this, [this, LocalConfirmResult]() { HandleButtonResult(LocalConfirmResult); }),
+			FCanExecuteAction::CreateSP(this, &SEditorUtilityDialog::CanExecuteHotkey)
+		)
+	);
+
+	const FVector2D DialogSize = InArgs._ClientSize;
+
+	SWindow::Construct(SWindow::FArguments()
+		.Title(InArgs._Title)
+		.ClientSize(DialogSize)
+		.SizingRule(DialogSize.IsNearlyZero() ? ESizingRule::Autosized : ESizingRule::UserSized)
+		.SupportsMaximize(false)
+		.SupportsMinimize(false)
+		[
+			SNew(SBorder)
+			.Padding(FMargin(4.f))
+			.BorderImage(FAppStyle::Get().GetBrush("ToolPanel.GroupBorder"))
+			[
+				SAssignNew(ContentBox, SVerticalBox)
+
+				+ SVerticalBox::Slot()
+				.FillHeight(1.0f)
+				[
+					SNullWidget::NullWidget
+				]
+
+				+ SVerticalBox::Slot()
+				.AutoHeight()
+				.HAlign(HAlign_Right)
+				.Padding(FMargin(0.f, 8.f, 0.f, 0.f))
+				[
+					CreateButtonBar(InArgs._ButtonLayout)
+				]
+			]
+		]
+	);
+
+	GetOnWindowClosedEvent().AddSPLambda(this, [this](const TSharedRef<SWindow>& ClosedWindow)
+	{
+		HandleWindowClosed(ClosedWindow);
+	});
+}
+
+void SEditorUtilityDialog::SetContent(UEditorUtilityDialogWidget* InWidget)
+{
+	check(InWidget);
+	HostedWidget.Reset(InWidget);
+
+	if (ContentBox.IsValid())
+	{
+		ContentBox->GetSlot(0)
+		[
+			InWidget->TakeWidget()
+		];
+	}
+}
+
+EAppReturnType::Type SEditorUtilityDialog::ShowModal()
+{
+	if (!FApp::IsUnattended() && !GIsRunningUnattendedScript && FSlateApplication::IsInitialized())
+	{
+		FSlateApplication::Get().AddModalWindow(
+			SharedThis(this),
+			FGlobalTabmanager::Get()->GetRootWindow());
+	}
+	else
+	{
+		// In unattended mode the window was never shown, so HandleWindowClosed won't fire.
+		// Fire widget close events manually to maintain symmetry with OnDialogOpened.
+		if (UEditorUtilityDialogWidget* Widget = HostedWidget.Get())
+		{
+			Widget->OnDialogButtonPressed(CancelResult);
+			Widget->OnDialogCancelled();
+		}
+		UE_LOG(LogEditorUtilityBlueprint, Warning, TEXT("SEditorUtilityDialog::ShowModal called in unattended mode, returning default Cancel."));
+	}
+
+	return LastPressedButton;
+}
+
+void SEditorUtilityDialog::Show(FOnEditorUtilityDialogSlateClosedNative InOnClosed)
+{
+	OnClosedDelegate = InOnClosed;
+
+	if (!FApp::IsUnattended() && !GIsRunningUnattendedScript && FSlateApplication::IsInitialized())
+	{
+		FSlateApplication::Get().AddWindow(
+			SharedThis(this), true);
+	}
+	else
+	{
+		OnClosedDelegate.ExecuteIfBound(LastPressedButton);
+		UE_LOG(LogEditorUtilityBlueprint, Warning, TEXT("SEditorUtilityDialog::Show called in unattended mode, firing OnClosed immediately."));
+	}
+}
+
+void SEditorUtilityDialog::CloseWithResult(EAppReturnType::Type Result)
+{
+	HandleButtonResult(Result);
+}
+
+FReply SEditorUtilityDialog::OnKeyDown(const FGeometry& MyGeometry, const FKeyEvent& InKeyEvent)
+{
+	if (CommandList->ProcessCommandBindings(InKeyEvent))
+	{
+		return FReply::Handled();
+	}
+	return SWindow::OnKeyDown(MyGeometry, InKeyEvent);
+}
+
+TSharedRef<SWidget> SEditorUtilityDialog::CreateButtonBar(EAppMsgType::Type MsgType)
+{
+	TSharedRef<SHorizontalBox> ButtonBox = SNew(SHorizontalBox);
+
+	auto AddButton = [this, &ButtonBox](const FText& Label, EAppReturnType::Type Result, bool bIsPrimary = false)
+	{
+		const FButtonStyle* Style = bIsPrimary
+			? &FAppStyle::Get().GetWidgetStyle<FButtonStyle>("PrimaryButton")
+			: &FAppStyle::Get().GetWidgetStyle<FButtonStyle>("Button");
+
+		ButtonBox->AddSlot()
+			.AutoWidth()
+			.Padding(FAppStyle::Get().GetMargin("StandardDialog.SlotPadding"))
+			[
+				SNew(SButton)
+				.ButtonStyle(Style)
+				.OnClicked_Lambda([this, Result]() { return OnButtonClicked(Result); })
+				[
+					SNew(STextBlock)
+					.Text(Label)
+				]
+			];
+	};
+
+	switch (MsgType)
+	{
+	case EAppMsgType::Ok:
+		AddButton(LOCTEXT("Ok", "OK"), EAppReturnType::Ok, true);
+		break;
+	case EAppMsgType::YesNo:
+		AddButton(LOCTEXT("Yes", "Yes"), EAppReturnType::Yes, true);
+		AddButton(LOCTEXT("No", "No"), EAppReturnType::No);
+		break;
+	case EAppMsgType::OkCancel:
+		AddButton(LOCTEXT("Ok", "OK"), EAppReturnType::Ok, true);
+		AddButton(LOCTEXT("Cancel", "Cancel"), EAppReturnType::Cancel);
+		break;
+	case EAppMsgType::YesNoCancel:
+		AddButton(LOCTEXT("Yes", "Yes"), EAppReturnType::Yes, true);
+		AddButton(LOCTEXT("No", "No"), EAppReturnType::No);
+		AddButton(LOCTEXT("Cancel", "Cancel"), EAppReturnType::Cancel);
+		break;
+	case EAppMsgType::CancelRetryContinue:
+		AddButton(LOCTEXT("Continue", "Continue"), EAppReturnType::Continue, true);
+		AddButton(LOCTEXT("Retry", "Retry"), EAppReturnType::Retry);
+		AddButton(LOCTEXT("Cancel", "Cancel"), EAppReturnType::Cancel);
+		break;
+	case EAppMsgType::YesNoYesAllNoAll:
+		AddButton(LOCTEXT("Yes", "Yes"), EAppReturnType::Yes, true);
+		AddButton(LOCTEXT("No", "No"), EAppReturnType::No);
+		AddButton(LOCTEXT("YesAll", "Yes All"), EAppReturnType::YesAll);
+		AddButton(LOCTEXT("NoAll", "No All"), EAppReturnType::NoAll);
+		break;
+	case EAppMsgType::YesNoYesAllNoAllCancel:
+		AddButton(LOCTEXT("Yes", "Yes"), EAppReturnType::Yes, true);
+		AddButton(LOCTEXT("No", "No"), EAppReturnType::No);
+		AddButton(LOCTEXT("YesAll", "Yes All"), EAppReturnType::YesAll);
+		AddButton(LOCTEXT("NoAll", "No All"), EAppReturnType::NoAll);
+		AddButton(LOCTEXT("Cancel", "Cancel"), EAppReturnType::Cancel);
+		break;
+	case EAppMsgType::YesNoYesAll:
+		AddButton(LOCTEXT("Yes", "Yes"), EAppReturnType::Yes, true);
+		AddButton(LOCTEXT("No", "No"), EAppReturnType::No);
+		AddButton(LOCTEXT("YesAll", "Yes All"), EAppReturnType::YesAll);
+		break;
+	default:
+		AddButton(LOCTEXT("Ok", "OK"), EAppReturnType::Ok, true);
+		break;
+	}
+
+	return ButtonBox;
+}
+
+FReply SEditorUtilityDialog::OnButtonClicked(EAppReturnType::Type Result)
+{
+	HandleButtonResult(Result);
+	return FReply::Handled();
+}
+
+void SEditorUtilityDialog::HandleButtonResult(EAppReturnType::Type Result)
+{
+	if (bIsClosing)
+	{
+		return;
+	}
+	bIsClosing = true;
+
+	LastPressedButton = Result;
+
+	if (UEditorUtilityDialogWidget* Widget = HostedWidget.Get())
+	{
+		Widget->OnDialogButtonPressed(Result);
+
+		if (IsPositiveResult(Result))
+		{
+			Widget->OnDialogConfirmed();
+		}
+		else
+		{
+			Widget->OnDialogCancelled();
+		}
+	}
+
+	RequestDestroyWindow();
+}
+
+void SEditorUtilityDialog::HandleWindowClosed(const TSharedRef<SWindow>& /*ClosedWindow*/)
+{
+	if (!bIsClosing)
+	{
+		bIsClosing = true;
+		LastPressedButton = CancelResult;
+
+		if (UEditorUtilityDialogWidget* Widget = HostedWidget.Get())
+		{
+			Widget->OnDialogButtonPressed(CancelResult);
+			Widget->OnDialogCancelled();
+		}
+	}
+
+	OnClosedDelegate.ExecuteIfBound(LastPressedButton);
+}
+
+bool SEditorUtilityDialog::CanExecuteHotkey() const
+{
+	return !bIsClosing;
+}
+
+bool SEditorUtilityDialog::IsPositiveResult(EAppReturnType::Type Result)
+{
+	switch (Result)
+	{
+	case EAppReturnType::Ok:
+	case EAppReturnType::Yes:
+	case EAppReturnType::YesAll:
+	case EAppReturnType::Retry:
+	case EAppReturnType::Continue:
+		return true;
+	default:
+		return false;
+	}
+}
+
+#undef LOCTEXT_NAMESPACE
diff --git a/Engine/Source/Editor/Blutility/Private/SEditorUtilityDialog.h b/Engine/Source/Editor/Blutility/Private/SEditorUtilityDialog.h
new file mode 100644
--- /dev/null
+++ b/Engine/Source/Editor/Blutility/Private/SEditorUtilityDialog.h
@@ -0,0 +1,57 @@
+// Copyright Epic Games, Inc. All Rights Reserved.
+
+#pragma once
+
+#include "CoreMinimal.h"
+#include "GenericPlatform/GenericPlatformMisc.h"
+#include "Framework/Commands/UICommandList.h"
+#include "UObject/StrongObjectPtr.h"
+#include "Widgets/SWindow.h"
+
+class UEditorUtilityDialogWidget;
+
+/** Delegate fired when an SEditorUtilityDialog is closed (used by the non-modal Show path). */
+DECLARE_DELEGATE_OneParam(FOnEditorUtilityDialogSlateClosedNative, EAppReturnType::Type /*Result*/);
+
+/**
+ * A modal (or non-modal) dialog window that hosts an Editor Utility Widget.
+ *
+ * Modeled after SCustomDialog but specialized for EUW content. Auto-generates
+ * a standard button bar from EAppMsgType and returns EAppReturnType results.
+ */
+class SEditorUtilityDialog : public SWindow
+{
+public:
+	SLATE_BEGIN_ARGS(SEditorUtilityDialog)
+		: _ButtonLayout(EAppMsgType::OkCancel)
+		, _ClientSize(FVector2D::ZeroVector)
+	{}
+		SLATE_ARGUMENT(FText, Title)
+		SLATE_ARGUMENT(EAppMsgType::Type, ButtonLayout)
+		SLATE_ARGUMENT(FVector2D, ClientSize)
+	SLATE_END_ARGS()
+
+	void Construct(const FArguments& InArgs);
+	void SetContent(UEditorUtilityDialogWidget* InWidget);
+	EAppReturnType::Type ShowModal();
+	void Show(FOnEditorUtilityDialogSlateClosedNative OnClosed);
+	void CloseWithResult(EAppReturnType::Type Result);
+	EAppReturnType::Type GetResult() const { return LastPressedButton; }
+	virtual FReply OnKeyDown(const FGeometry& MyGeometry, const FKeyEvent& InKeyEvent) override;
+
+private:
+	TSharedRef<SWidget> CreateButtonBar(EAppMsgType::Type MsgType);
+	FReply OnButtonClicked(EAppReturnType::Type Result);
+	void HandleButtonResult(EAppReturnType::Type Result);
+	void HandleWindowClosed(const TSharedRef<SWindow>& ClosedWindow);
+	bool CanExecuteHotkey() const;
+	static bool IsPositiveResult(EAppReturnType::Type Result);
+
+	EAppReturnType::Type LastPressedButton = EAppReturnType::Cancel;
+	EAppReturnType::Type CancelResult = EAppReturnType::Cancel;
+	bool bIsClosing = false;
+	TStrongObjectPtr<UEditorUtilityDialogWidget> HostedWidget;
+	TSharedPtr<SVerticalBox> ContentBox;
+	TSharedPtr<FUICommandList> CommandList;
+	FOnEditorUtilityDialogSlateClosedNative OnClosedDelegate;
+};
diff --git a/Engine/Source/Editor/Blutility/Public/EditorUtilityDialogLibrary.h b/Engine/Source/Editor/Blutility/Public/EditorUtilityDialogLibrary.h
new file mode 100644
--- /dev/null
+++ b/Engine/Source/Editor/Blutility/Public/EditorUtilityDialogLibrary.h
@@ -0,0 +1,53 @@
+// Copyright Epic Games, Inc. All Rights Reserved.
+
+#pragma once
+
+#include "CoreMinimal.h"
+#include "Kismet/BlueprintFunctionLibrary.h"
+#include "GenericPlatform/GenericPlatformMisc.h"
+#include "EditorUtilityDialogResult.h"
+#include "EditorUtilityDialogWidget.h"
+
+#include "EditorUtilityDialogLibrary.generated.h"
+
+class SEditorUtilityDialog;
+
+#define UE_API BLUTILITY_API
+
+/**
+ * Blueprint Function Library for creating modal and non-modal dialogs
+ * that contain Editor Utility Widgets.
+ */
+UCLASS(MinimalAPI)
+class UEditorUtilityDialogLibrary : public UBlueprintFunctionLibrary
+{
+	GENERATED_BODY()
+
+public:
+	/** Creates an Editor Utility Dialog Widget and shows it as a modal window. Blocks until dismissed. */
+	UFUNCTION(BlueprintCallable, Category = "Editor Utility|Dialog")
+	static UE_API FEditorUtilityDialogResult ShowModalEditorUtilityDialog(TSubclassOf<UEditorUtilityDialogWidget> WidgetClass, FText Title, EAppMsgType::Type ButtonLayout = EAppMsgType::OkCancel, FVector2D Size = FVector2D::ZeroVector);
+
+	/** Creates an Editor Utility Dialog Widget as a non-modal window. Returns immediately. OnClosed fires when dismissed. */
+	UFUNCTION(BlueprintCallable, Category = "Editor Utility|Dialog", meta = (AutoCreateRefTerm = "OnClosed", DeterminesOutputType = "WidgetClass"))
+	static UE_API UEditorUtilityDialogWidget* ShowEditorUtilityDialog(TSubclassOf<UEditorUtilityDialogWidget> WidgetClass, const FOnEditorUtilityDialogClosed& OnClosed, FText Title = FText::GetEmpty(), EAppMsgType::Type ButtonLayout = EAppMsgType::OkCancel, FVector2D Size = FVector2D::ZeroVector);
+
+	/** Creates and initializes a dialog widget without showing it. Set data on it, then pass to Show Modal/Non-Modal (Instance). */
+	UFUNCTION(BlueprintCallable, Category = "Editor Utility|Dialog", meta = (DeterminesOutputType = "WidgetClass"))
+	static UE_API UEditorUtilityDialogWidget* CreateEditorUtilityDialogWidget(TSubclassOf<UEditorUtilityDialogWidget> WidgetClass);
+
+	/** Shows a pre-constructed Editor Utility Dialog Widget as a modal window. */
+	UFUNCTION(BlueprintCallable, Category = "Editor Utility|Dialog")
+	static UE_API FEditorUtilityDialogResult ShowModalEditorUtilityDialogInstance(UEditorUtilityDialogWidget* WidgetInstance, FText Title, EAppMsgType::Type ButtonLayout = EAppMsgType::OkCancel, FVector2D Size = FVector2D::ZeroVector);
+
+	/** Shows a pre-constructed Editor Utility Dialog Widget as a non-modal window. */
+	UFUNCTION(BlueprintCallable, Category = "Editor Utility|Dialog", meta = (AutoCreateRefTerm = "OnClosed"))
+	static UE_API void ShowEditorUtilityDialogInstance(UEditorUtilityDialogWidget* WidgetInstance, const FOnEditorUtilityDialogClosed& OnClosed, FText Title = FText::GetEmpty(), EAppMsgType::Type ButtonLayout = EAppMsgType::OkCancel, FVector2D Size = FVector2D::ZeroVector);
+
+private:
+	static UE_API TPair<UEditorUtilityDialogWidget*, TSharedPtr<SEditorUtilityDialog>> CreateDialogInternal(TSubclassOf<UEditorUtilityDialogWidget> WidgetClass, const FText& Title, EAppMsgType::Type ButtonLayout, const FVector2D& Size);
+
+	static UE_API TSharedPtr<SEditorUtilityDialog> WrapWidgetInDialog(UEditorUtilityDialogWidget* Widget, const FText& Title, EAppMsgType::Type ButtonLayout, const FVector2D& Size);
+};
+
+#undef UE_API
diff --git a/Engine/Source/Editor/Blutility/Public/EditorUtilityDialogResult.h b/Engine/Source/Editor/Blutility/Public/EditorUtilityDialogResult.h
new file mode 100644
--- /dev/null
+++ b/Engine/Source/Editor/Blutility/Public/EditorUtilityDialogResult.h
@@ -0,0 +1,36 @@
+// Copyright Epic Games, Inc. All Rights Reserved.
+
+#pragma once
+
+#include "CoreMinimal.h"
+#include "GenericPlatform/GenericPlatformMisc.h"
+#include "UObject/Object.h"
+
+#include "EditorUtilityDialogResult.generated.h"
+
+class UEditorUtilityDialogWidget;
+
+/** Delegate fired when an async (non-modal) Editor Utility Dialog is closed. */
+DECLARE_DYNAMIC_DELEGATE_TwoParams(
+	FOnEditorUtilityDialogClosed,
+	EAppReturnType::Type, Result,
+	UEditorUtilityDialogWidget*, WidgetInstance);
+
+/**
+ * Result returned by ShowModalEditorUtilityDialog.
+ * Contains the button the user pressed and a reference to the widget instance
+ * so the caller can read back any data the user entered.
+ */
+USTRUCT(MinimalAPI, BlueprintType)
+struct FEditorUtilityDialogResult
+{
+	GENERATED_BODY()
+
+	/** Which button the user pressed to close the dialog (Ok, Cancel, Yes, No, etc.). */
+	UPROPERTY(BlueprintReadOnly, Category = "Dialog Result")
+	TEnumAsByte<EAppReturnType::Type> ReturnType = EAppReturnType::Cancel;
+
+	/** The widget instance shown in the dialog. Cast to your subclass to read back user-entered data. */
+	UPROPERTY(BlueprintReadOnly, Category = "Dialog Result")
+	TObjectPtr<UEditorUtilityDialogWidget> WidgetInstance = nullptr;
+};
diff --git a/Engine/Source/Editor/Blutility/Public/EditorUtilityDialogWidget.h b/Engine/Source/Editor/Blutility/Public/EditorUtilityDialogWidget.h
new file mode 100644
--- /dev/null
+++ b/Engine/Source/Editor/Blutility/Public/EditorUtilityDialogWidget.h
@@ -0,0 +1,71 @@
+// Copyright Epic Games, Inc. All Rights Reserved.
+
+#pragma once
+
+#include "CoreMinimal.h"
+#include "EditorUtilityWidget.h"
+#include "GenericPlatform/GenericPlatformMisc.h"
+
+#include "EditorUtilityDialogWidget.generated.h"
+
+class SEditorUtilityDialog;
+
+#define UE_API BLUTILITY_API
+
+/**
+ * Base class for Editor Utility Widgets that are designed to be shown in modal or
+ * non-modal dialog windows. Subclass this in Blueprint to create custom dialog content.
+ *
+ * Unlike regular Editor Utility Widgets (which live in dockable tabs), dialog widgets
+ * are shown in standalone windows with a standard button bar (OK, Cancel, Yes/No, etc.).
+ *
+ * The caller receives this widget instance after the dialog closes, so you can expose
+ * UPROPERTY fields that the caller reads back to get user-entered data.
+ */
+UCLASS(MinimalAPI, Abstract, Blueprintable, meta = (ShowWorldContextPin))
+class UEditorUtilityDialogWidget : public UEditorUtilityWidget
+{
+	GENERATED_BODY()
+
+public:
+	/** The title shown in the modal window's title bar. Can be overridden by the caller when showing the dialog. */
+	UPROPERTY(EditDefaultsOnly, BlueprintReadWrite, Category = "Dialog Settings")
+	FText DialogTitle;
+
+	/** Default window size in pixels (Width, Height). Set to (0, 0) to auto-size to content. */
+	UPROPERTY(EditDefaultsOnly, BlueprintReadWrite, Category = "Dialog Settings")
+	FVector2D DesiredDialogSize = FVector2D::ZeroVector;
+
+	/** Which standard buttons appear at the bottom of the dialog (Ok, OkCancel, YesNo, etc.). */
+	UPROPERTY(EditDefaultsOnly, BlueprintReadWrite, Category = "Dialog Settings")
+	TEnumAsByte<EAppMsgType::Type> ButtonLayout = EAppMsgType::OkCancel;
+
+	/** Called after the widget is created, before the dialog becomes visible. Initialize state here. */
+	UFUNCTION(BlueprintImplementableEvent, Category = "Dialog|Events")
+	void OnDialogOpened();
+
+	/** Called when the user presses a positive button (OK, Yes, etc.). Fires before close. */
+	UFUNCTION(BlueprintImplementableEvent, Category = "Dialog|Events")
+	void OnDialogConfirmed();
+
+	/** Called when the user presses Cancel, No, or closes via the X button. Fires before close. */
+	UFUNCTION(BlueprintImplementableEvent, Category = "Dialog|Events")
+	void OnDialogCancelled();
+
+	/** Called for ANY button press. Catch-all alternative to OnDialogConfirmed/OnDialogCancelled. */
+	UFUNCTION(BlueprintImplementableEvent, Category = "Dialog|Events")
+	void OnDialogButtonPressed(EAppReturnType::Type Result);
+
+	/** Programmatically close this dialog with a specific result. */
+	UFUNCTION(BlueprintCallable, Category = "Dialog")
+	UE_API void CloseDialog(EAppReturnType::Type Result);
+
+	/** Set by the dialog system when the widget is hosted in an SEditorUtilityDialog. */
+	void SetOwningDialog(TSharedPtr<SEditorUtilityDialog> InDialog);
+
+private:
+	/** Weak reference to the Slate dialog window hosting this widget. Used by CloseDialog(). */
+	TWeakPtr<SEditorUtilityDialog> OwningDialog;
+};
+
+#undef UE_API
